Decoupled StudioTelemetryEditor from StudioTelemetry and moved it into its own Plugin EditorTelemetry
This makes StudioTelemetry lighter weight and no longer dependant on the Engine and Editor leaving more useful for applying to other systems or applications outside of the Editor or Client
Added OnStartSession, OnEndSession and OnRecordEvent delegates that are used by EditorTelemetry to Start/Stop its own session at the time StudioTelemetry starts and stops.

diff --git a/Engine/Plugins/Experimental/EditorTelemetry/EditorTelemetry.uplugin b/Engine/Plugins/Experimental/EditorTelemetry/EditorTelemetry.uplugin
new file mode 100644
--- /dev/null
+++ b/Engine/Plugins/Experimental/EditorTelemetry/EditorTelemetry.uplugin
@@ -0,0 +1,30 @@
+{
+	"FileVersion": 1,
+	"Version": 1,
+	"VersionName": "1.0",
+	"FriendlyName": "Editor Performance",
+	"Description": "Plugin that provides Editor Performance feedback to developers",
+	"Category": "Performance",
+	"CreatedBy": "Epic Games, Inc.",
+	"CreatedByURL": "https://epicgames.com",
+	"DocsURL": "",
+	"MarketplaceURL": "",
+	"SupportURL": "",
+	"EnabledByDefault": true,
+	"CanContainContent": false,
+	"IsBetaVersion": false,
+	"Installed": true,
+	"Modules": [
+		{
+			"Name": "EditorTelemetry",
+			"Type": "Editor",
+			"LoadingPhase": "PostConfigInit"
+		}
+	],
+	"Plugins": [
+		{
+			"Name": "StudioTelemetry",
+			"Enabled": true
+		}
+	]
+}
\ No newline at end of file
diff --git a/Engine/Plugins/Experimental/EditorTelemetry/Source/EditorTelemetry.Build.cs b/Engine/Plugins/Experimental/EditorTelemetry/Source/EditorTelemetry.Build.cs
new file mode 100644
--- /dev/null
+++ b/Engine/Plugins/Experimental/EditorTelemetry/Source/EditorTelemetry.Build.cs
@@ -0,0 +1,34 @@
+// Copyright Epic Games, Inc. All Rights Reserved.
+
+namespace UnrealBuildTool.Rules
+{
+	public class EditorTelemetry : ModuleRules
+	{
+		public EditorTelemetry(ReadOnlyTargetRules Target) : base(Target)
+		{
+			PublicDependencyModuleNames.AddRange(
+				new string[] {
+					"Core"
+				}
+			);
+
+			PrivateDependencyModuleNames.AddRange(
+				new string[] {
+					"CoreUObject",
+					"Analytics",
+					"TelemetryUtils",
+					"AssetRegistry",
+					"Engine",
+					"UnrealEd",
+					"HTTP",
+					"StudioTelemetry",
+					"DerivedDataCache",
+					"Zen",
+					"IoStoreOnDemand",
+					"ContentBrowser",
+					"ContentBrowserData",
+				}
+			);
+		}
+	}
+}
diff --git a/Engine/Plugins/Experimental/EditorTelemetry/Source/Private/EditorTelemetry.cpp b/Engine/Plugins/Experimental/EditorTelemetry/Source/Private/EditorTelemetry.cpp
new file mode 100644
--- /dev/null
+++ b/Engine/Plugins/Experimental/EditorTelemetry/Source/Private/EditorTelemetry.cpp
@@ -0,0 +1,1047 @@
+// Copyright Epic Games, Inc. All Rights Reserved.
+
+#include "EditorTelemetry.h"
+
+#include "AnalyticsTracer.h"
+#include "AssetRegistry/AssetRegistryTelemetry.h"
+#include "CollectionManagerModule.h"
+#include "ContentBrowserModule.h"
+#include "ContentBrowserTelemetry.h"
+#include "CookOnTheSide/CookOnTheFlyServer.h"
+#include "DerivedDataCacheInterface.h"
+#include "DerivedDataCacheUsageStats.h"
+#include "Editor.h"
+#include "Editor/UnrealEdEngine.h"
+#include "Engine/AssetManager.h"
+#include "Experimental/ZenServerInterface.h"
+#include "FileHelpers.h"
+#include "HttpManager.h"
+#include "HttpModule.h"
+#include "IO/IoStoreOnDemand.h"
+#include "Misc/FeedbackContext.h"
+#include "Modules/ModuleManager.h"
+#include "Policies/CondensedJsonPrintPolicy.h"
+#include "ProfilingDebugging/CookStats.h"
+#include "ShaderStats.h"
+#include "StudioTelemetry.h"
+#include "Subsystems/AssetEditorSubsystem.h"
+#include "TelemetryRouter.h"
+#include "UObject/ICookInfo.h"
+#include "UnrealEdGlobals.h"
+#include "Virtualization/VirtualizationSystem.h"
+
+UE_DISABLE_OPTIMIZATION_SHIP
+
+namespace Private
+{
+	const FName ContentBrowserModuleName = TEXT("ContentBrowser");
+	
+	// Json writer subclass to allow us to avoid using a SharedPtr to write basic Json.
+	typedef TCondensedJsonPrintPolicy<TCHAR> FPrintPolicy;
+	class FAnalyticsJsonWriter : public TJsonStringWriter<FPrintPolicy>
+	{
+	public:
+		explicit FAnalyticsJsonWriter(FString* Out) : TJsonStringWriter<FPrintPolicy>(Out, 0)
+		{
+		}
+	};
+}
+
+const TCHAR* LexToString(ECollectionTelemetryAssetAddedWorkflow Enum)
+{
+	switch(Enum)
+	{
+	case ECollectionTelemetryAssetAddedWorkflow::ContextMenu: return TEXT("ContextMenu");
+	case ECollectionTelemetryAssetAddedWorkflow::DragAndDrop: return TEXT("DragAndDrop");
+	default: return TEXT("");
+	}
+}
+	
+const TCHAR* LexToString(ECollectionTelemetryAssetRemovedWorkflow Enum)
+{
+	switch(Enum)
+	{
+	case ECollectionTelemetryAssetRemovedWorkflow::ContextMenu: return TEXT("ContextMenu");
+	default: return TEXT("");
+	}
+}
+
+template<typename T>
+FString AnalyticsOptionalToStringOrNull(const TOptional<T>& Opt)
+{
+	return Opt.IsSet() ? AnalyticsConversionToString(Opt.GetValue()) : FString(TEXT("null"));
+}
+
+FEditorTelemetry& FEditorTelemetry::Get()
+{
+	static FEditorTelemetry StudioTelemetryEditorInstance = FEditorTelemetry();
+	return StudioTelemetryEditorInstance;
+}
+
+void FEditorTelemetry::RecordEvent_Cooking(TArray<FAnalyticsEventAttribute> Attributes)
+{
+#if ENABLE_COOK_STATS
+
+	const int SchemaVersion = 3;
+
+	Attributes.Emplace(TEXT("SchemaVersion"), SchemaVersion);
+
+	TMap<FString,FAnalyticsEventAttribute> CookAttributes;
+	
+	// Sends each cook stat to the studio analytics system.
+	auto GatherAnalyticsAttributes = [&CookAttributes](const FString& StatName, const TArray<FCookStatsManager::StringKeyValue>& StatAttributes)
+	{
+		for (const auto& Attr : StatAttributes)
+		{
+			const FString FormattedAttrName = (StatName + "_" + Attr.Key).Replace(TEXT("."), TEXT("_"));
+
+			if (CookAttributes.Find(FormattedAttrName)==nullptr)
+			{
+				CookAttributes.Emplace(FormattedAttrName, Attr.Value.IsNumeric() ? FAnalyticsEventAttribute(FormattedAttrName, FCString::Atof(*Attr.Value)) : FAnalyticsEventAttribute(FormattedAttrName, Attr.Value));
+			}
+		}
+	};
+
+	// Now actually grab the stats 
+	FCookStatsManager::LogCookStats(GatherAnalyticsAttributes);
+
+	// Add the values to the attributes
+	for (TMap<FString, FAnalyticsEventAttribute>::TConstIterator it(CookAttributes); it; ++it)
+	{
+		Attributes.Emplace((*it).Value);
+	}
+
+	// Gather the DDC summary stats
+	FDerivedDataCacheSummaryStats SummaryStats;
+
+	GatherDerivedDataCacheSummaryStats(SummaryStats);
+
+	// Append to the attributes
+	for (const FDerivedDataCacheSummaryStat& Stat : SummaryStats.Stats)
+	{
+		FString AttributeName = TEXT("DDC_Summary") + Stat.Key.Replace(TEXT("."), TEXT("_"));
+
+		if (Stat.Value.IsNumeric())
+		{
+			Attributes.Emplace(AttributeName, FCString::Atof(*Stat.Value));
+		}
+		else
+		{
+			Attributes.Emplace(AttributeName, Stat.Value);
+		}
+	}
+
+#if UE_WITH_ZEN
+	// Gather Zen analytics
+	if (UE::Zen::IsDefaultServicePresent())
+	{
+		UE::Zen::GetDefaultServiceInstance().GatherAnalytics(Attributes);
+	}
+#endif
+
+	if (UE::Virtualization::IVirtualizationSystem::Get().IsEnabled())
+	{
+		// Gather Virtualization analytics
+		UE::Virtualization::IVirtualizationSystem::Get().GatherAnalytics(Attributes);
+	}
+
+	FShaderStatsFunctions::GatherShaderAnalytics(Attributes);
+	
+	FStudioTelemetry::Get().RecordEvent(TEXT("Core.Cooking"), Attributes);
+#endif
+}
+
+void FEditorTelemetry::RecordEvent_Loading(const FString& LoadingName, double LoadingSeconds, TArray<FAnalyticsEventAttribute> Attributes )
+{
+	const int SchemaVersion = 3;
+
+	Attributes.Emplace(TEXT("SchemaVersion"), SchemaVersion);
+	Attributes.Emplace(TEXT("LoadingName"), LoadingName);
+	Attributes.Emplace(TEXT("LoadingSeconds"), LoadingSeconds);
+	
+#if ENABLE_COOK_STATS
+
+#if UE_WITH_ZEN
+	// Gather Zen analytics
+	if (UE::Zen::IsDefaultServicePresent())
+	{
+		UE::Zen::GetDefaultServiceInstance().GatherAnalytics(Attributes);
+	}
+#endif
+
+	if (UE::Virtualization::IVirtualizationSystem::Get().IsEnabled())
+	{
+		// Gather Virtualization analytics
+		UE::Virtualization::IVirtualizationSystem::Get().GatherAnalytics(Attributes);
+	}
+
+	// Gather the DDC summary stats
+	FDerivedDataCacheSummaryStats SummaryStats;
+
+	GatherDerivedDataCacheSummaryStats(SummaryStats);
+
+	// Append to the attributes
+	for (const FDerivedDataCacheSummaryStat& Stat : SummaryStats.Stats)
+	{
+		FString AttributeName = TEXT("DDC_Summary_") + Stat.Key.Replace(TEXT("."), TEXT("_"));
+
+		if (Stat.Value.IsNumeric())
+		{
+			Attributes.Emplace(AttributeName, FCString::Atof(*Stat.Value));
+		}
+		else
+		{
+			Attributes.Emplace(AttributeName, Stat.Value);
+		}
+	}
+#endif			
+
+	FStudioTelemetry::Get().RecordEvent(TEXT("Core.Loading"), Attributes);
+}
+
+void FEditorTelemetry::RecordEvent_DDCResource(const FString& Context, TArray<FAnalyticsEventAttribute> Attributes)
+{
+#if ENABLE_COOK_STATS
+	// Gather the latest resource stats
+	TArray<FDerivedDataCacheResourceStat> ResourceStats;
+
+	GatherDerivedDataCacheResourceStats(ResourceStats);
+
+	const int SchemaVersion = 3;
+
+	Attributes.Emplace(TEXT("SchemaVersion"), SchemaVersion);
+	Attributes.Emplace(TEXT("LoadingName"), Context);
+	
+	// Send a resource event per asset type
+	for (const FDerivedDataCacheResourceStat& Stat : ResourceStats)
+	{
+		const double TotalTimeSec = Stat.BuildTimeSec + Stat.LoadTimeSec;
+		const int64 TotalSizeMB = Stat.BuildSizeMB + Stat.LoadSizeMB;
+
+		if (Stat.AssetType.IsEmpty() || Stat.TotalCount==0)
+		{
+			// Empty asset type or nothing was built or loaded for this type
+			continue;
+		}
+	
+		TArray<FAnalyticsEventAttribute> EventAttributes = Attributes;
+
+		EventAttributes.Emplace(TEXT("AssetType"), Stat.AssetType);
+		EventAttributes.Emplace(TEXT("Load_Count"), Stat.LoadCount);
+		EventAttributes.Emplace(TEXT("Load_TimeSec"), Stat.LoadTimeSec);
+		EventAttributes.Emplace(TEXT("Load_SizeMB"), Stat.LoadSizeMB);
+		EventAttributes.Emplace(TEXT("Build_Count"), Stat.BuildCount);
+		EventAttributes.Emplace(TEXT("Build_TimeSec"), Stat.BuildTimeSec);
+		EventAttributes.Emplace(TEXT("Build_SizeMB"), Stat.BuildSizeMB);
+		EventAttributes.Emplace(TEXT("Total_Count"), Stat.TotalCount);
+		EventAttributes.Emplace(TEXT("Total_TimeSec"), TotalTimeSec);
+		EventAttributes.Emplace(TEXT("Total_SizeMB"), TotalSizeMB);
+		EventAttributes.Emplace(TEXT("Efficiency"), Stat.Efficiency);
+		EventAttributes.Emplace(TEXT("Thread_TimeSec"), Stat.GameThreadTimeSec);
+
+		FStudioTelemetry::Get().RecordEvent(TEXT("Core.DDC.Resource"), EventAttributes);
+	}
+#endif			
+}
+
+
+void FEditorTelemetry::RecordEvent_DDCSummary(const FString& Context, TArray<FAnalyticsEventAttribute> Attributes)
+{
+#if ENABLE_COOK_STATS
+	const int SchemaVersion = 3;
+
+	Attributes.Emplace(TEXT("SchemaVersion"), SchemaVersion);
+	Attributes.Emplace(TEXT("LoadingName"), Context);
+
+	// Gather the summary stats
+	FDerivedDataCacheSummaryStats SummaryStats;
+
+	GatherDerivedDataCacheSummaryStats(SummaryStats);
+
+	// Append to the attributes
+	for (const FDerivedDataCacheSummaryStat& Stat : SummaryStats.Stats)
+	{
+		FString AttributeName = Stat.Key.Replace(TEXT("."), TEXT("_"));
+
+		if (Stat.Value.IsNumeric())
+		{
+			Attributes.Emplace(AttributeName, FCString::Atof(*Stat.Value));
+		}
+		else
+		{
+			Attributes.Emplace(AttributeName, Stat.Value);
+		}
+	}
+
+	FStudioTelemetry::Get().RecordEvent(TEXT("Core.DDC.Summary"), Attributes);
+#endif			
+}
+
+void FEditorTelemetry::RecordEvent_IAS(const FString& Context, TArray<FAnalyticsEventAttribute> Attributes)
+{
+	// Gather the summary stats
+	FDerivedDataCacheSummaryStats SummaryStats;
+
+	using namespace UE::IoStore;
+	if (FIoStoreOnDemandModule* OnDemandModule = FModuleManager::Get().GetModulePtr<FIoStoreOnDemandModule>("IoStoreOnDemand"))
+	{
+		if (OnDemandModule->IsEnabled())
+		{
+			const int SchemaVersion = 1;
+
+			Attributes.Emplace(TEXT("SchemaVersion"), SchemaVersion);
+			Attributes.Emplace(TEXT("LoadingName"), Context);
+
+			OnDemandModule->ReportAnalytics(Attributes);
+			FStudioTelemetry::Get().RecordEvent(TEXT("Core.IAS"), Attributes);
+		}		
+	}
+}
+
+void FEditorTelemetry::RecordEvent_Zen(const FString& Context, TArray<FAnalyticsEventAttribute> Attributes)
+{
+#if UE_WITH_ZEN
+	// Gather Zen analytics
+	if (UE::Zen::IsDefaultServicePresent())
+	{
+		const int SchemaVersion = 1;
+
+		Attributes.Emplace(TEXT("SchemaVersion"), SchemaVersion);
+		Attributes.Emplace(TEXT("LoadingName"), Context);
+
+		UE::Zen::GetDefaultServiceInstance().GatherAnalytics(Attributes);
+		FStudioTelemetry::Get().RecordEvent(TEXT("Core.Zen"), Attributes);
+	}
+#endif
+}
+
+void FEditorTelemetry::RecordEvent_VirtualAssets(const FString& Context, TArray<FAnalyticsEventAttribute> Attributes)
+{
+	if (UE::Virtualization::IVirtualizationSystem::Get().IsEnabled())
+	{
+		const int SchemaVersion = 1;
+
+		Attributes.Emplace(TEXT("SchemaVersion"), SchemaVersion);
+		Attributes.Emplace(TEXT("LoadingName"), Context);
+
+		// Gather Virtualization analytics
+		UE::Virtualization::IVirtualizationSystem::Get().GatherAnalytics(Attributes);
+
+		FStudioTelemetry::Get().RecordEvent(TEXT("Core.VirtualAssets"), Attributes);
+	}
+}
+
+void FEditorTelemetry::RecordEvent_CoreSystems(const FString& Context, TArray<FAnalyticsEventAttribute> Attributes)
+{
+	FEditorTelemetry::RecordEvent_DDCResource(Context, Attributes);
+	FEditorTelemetry::RecordEvent_DDCSummary(Context, Attributes);
+	FEditorTelemetry::RecordEvent_IAS(Context, Attributes);
+	FEditorTelemetry::RecordEvent_Zen(Context, Attributes);
+	FEditorTelemetry::RecordEvent_VirtualAssets(Context, Attributes);
+}
+
+void FEditorTelemetry::RegisterCollectionWorkflowDelegates(FTelemetryRouter& Router)
+{
+	Router.OnTelemetry<FAssetAddedToCollectionTelemetryEvent>([](const FAssetAddedToCollectionTelemetryEvent& Event)
+	{
+		const int SchemaVersion = 1;
+		
+		FStudioTelemetry::Get().RecordEvent(TEXT("Editor.Collections.AssetsAdded"),
+		{
+			{ TEXT("SchemaVersion"), SchemaVersion},
+			{ TEXT("DurationSec"), Event.DurationSec},
+			{ TEXT("ObjectCount"), Event.NumAdded},
+			{ TEXT("Workflow"), Event.Workflow},
+			{ TEXT("CollectionShareType"), ECollectionShareType::ToString(Event.CollectionShareType)},
+		});
+	});
+
+	Router.OnTelemetry<FAssetRemovedFromCollectionTelemetryEvent>([](const FAssetRemovedFromCollectionTelemetryEvent& Event)
+	{
+		const int SchemaVersion = 1;
+		
+		FStudioTelemetry::Get().RecordEvent(TEXT("Editor.Collections.AssetsRemoved"),
+		{
+			{ TEXT("SchemaVersion"), SchemaVersion},
+			{ TEXT("DurationSec"), Event.DurationSec},
+			{ TEXT("ObjectCount"), Event.NumRemoved},
+			{ TEXT("Workflow"), Event.Workflow},
+			{ TEXT("CollectionShareType"), ECollectionShareType::ToString(Event.CollectionShareType) },
+		});
+	});
+
+	Router.OnTelemetry<FCollectionCreatedTelemetryEvent>([](const FCollectionCreatedTelemetryEvent& Event)
+	{
+		const int SchemaVersion = 1;
+
+		FStudioTelemetry::Get().RecordEvent(TEXT("Editor.Collections.CollectionCreated"),
+		{
+			{ TEXT("SchemaVersion"), SchemaVersion},
+			{ TEXT("DurationSec"), Event.DurationSec},
+			{ TEXT("CollectionShareType"), ECollectionShareType::ToString(Event.CollectionShareType)},
+		});
+	});
+
+	Router.OnTelemetry<FCollectionsDeletedTelemetryEvent>([](const FCollectionsDeletedTelemetryEvent& Event)
+	{
+		const int SchemaVersion = 1;
+
+		FStudioTelemetry::Get().RecordEvent(TEXT("Editor.Collections.CollectionDeleted"),
+		{
+			{ TEXT("SchemaVersion"), SchemaVersion},
+			{ TEXT("DurationSec"), Event.DurationSec},
+			{ TEXT("ObjectCount"), Event.CollectionsDeleted},
+		});
+	});
+}
+
+extern ENGINE_API float GAverageFPS;
+
+void FEditorTelemetry::HitchSamplerCallback()
+{
+	// Only sample framerate when we have focus
+	if (FApp::HasFocus())
+	{
+		// Sample a rolling average of FPS 
+		HitchAvergageFPS = (HitchAvergageFPS * HitchSampleCount + GAverageFPS) / (double)(HitchSampleCount + 1);
+		HitchSampleCount++;
+	}
+}
+	
+void FEditorTelemetry::HeartbeatCallback()
+{
+	if (HitchSampleCount>0)
+	{
+		// Hitching is when FPS is below our threshold
+		const bool IsHitching = HitchAvergageFPS < MinFPSForHitching;
+		static uint32 HitchCount = 0;
+		HitchCount += IsHitching ? 1 : 0;
+
+		if (IsHitching == false && HitchingSpan.IsValid() == true)
+		{
+			// No longer hitching and we have started a hitch span
+			TArray<FAnalyticsEventAttribute> Attributes;
+			Attributes.Emplace(FAnalyticsEventAttribute(TEXT("Hitch_Count"), HitchCount));
+			Attributes.Emplace(FAnalyticsEventAttribute(TEXT("Hitch_HitchesPerSecond"), (float)HitchCount / HitchingSpan->GetDuration()));
+			Attributes.Emplace(FAnalyticsEventAttribute(TEXT("Hitch_AverageFPS"), HitchAvergageFPS));
+
+			// End the hitch Span
+			FStudioTelemetry::Get().EndSpan(HitchingSpan, Attributes);
+
+			// Record the hitch event
+			FStudioTelemetry::Get().RecordEvent(TEXT("Core.Hitch"), Attributes);
+
+			Attributes.Emplace(TEXT("MapName"), EditorMapName);
+			Attributes.Emplace(TEXT("PIE_MapName"), PIEMapName);
+
+			// Record core systems events for the hitch
+			RecordEvent_CoreSystems(TEXT("Hitch"));
+
+			// No longer need the hitch span for now so reset it
+			HitchingSpan.Reset();
+		}
+		else if (IsHitching == true && HitchingSpan.IsValid() == false)
+		{
+			// We are hitching and we have not started a hitch span
+			HitchingSpan = FStudioTelemetry::Get().StartSpan(HitchingSpanName);
+			HitchCount = 1;
+		}
+
+		// Reset the hitch sampler
+		HitchSampleCount = 0;
+		HitchAvergageFPS = 0;
+	}
+}
+
+void FEditorTelemetry::StartSession()
+{
+	if (FStudioTelemetry::Get().IsAvailable() == false)
+	{
+		return;
+	}
+
+	SessionStartTime = FPlatformTime::Seconds();
+
+	// Install Editor Mode callbacks
+	
+	// Start Editor and Editor Boot span. Note : this will only start when the plugin is loaded and as such will miss any activity that runs beforehand
+	EditorSpan = FStudioTelemetry::Get().StartSpan(EditorSpanName);
+	EditorBootSpan = FStudioTelemetry::Get().StartSpan(EditorBootSpanName, EditorSpan);
+	EditorMapName = TEXT("None");
+
+	TArray<FAnalyticsEventAttribute> Attributes;
+	Attributes.Emplace(FAnalyticsEventAttribute(TEXT("MapName"), EditorMapName));
+	EditorBootSpan->AddAttributes(Attributes);
+
+	// Set up Slow Task callbacks
+	ensureMsgf(GWarn, TEXT("GWarn was not valid"));
+
+	if (GWarn != nullptr)
+	{
+		// Start the SlowTask span
+		GWarn->OnStartSlowTaskWithGuid().AddLambda([this](FGuid TaskGuid, const FText& TaskName)
+			{
+				// Slow tasks can possibly be started from multiple threads, so we need to protect the registered span table
+				FScopeLock ScopeLock(&TaskSpanCriticalSection);
+
+				TSharedPtr<IAnalyticsSpan>* SpanPtr = TaskSpans.Find(TaskGuid);
+
+				// Only one task with this Guid is running asynchronously is supported at this time.
+				if (SpanPtr == nullptr)
+				{
+					TArray<FAnalyticsEventAttribute> Attributes;
+					Attributes.Emplace(TEXT("MapName"), EditorMapName);
+					Attributes.Emplace(TEXT("TaskName"), TaskName.ToString());
+
+					// Create and start a new slow task span
+					TSharedPtr<IAnalyticsSpan> SlowTaskSpan = FStudioTelemetry::Get().StartSpan(TEXT("SlowTask"), Attributes);
+
+					// Store this SlowTask span so we can find it when it finishes
+					TaskSpans.Add(TaskGuid, SlowTaskSpan);
+				}
+
+				TRACE_BEGIN_REGION(*TaskName.ToString());
+			});
+
+		// End the SlowTask span
+		GWarn->OnFinalizeSlowTaskWithGuid().AddLambda([this](FGuid TaskGuid, const FText& TaskName)
+			{
+				TRACE_END_REGION(*TaskName.ToString());
+
+				// Slow tasks can possibly be finalized from multiple threads, so we need to protect the registered span table
+				FScopeLock ScopeLock(&TaskSpanCriticalSection);
+
+				// Find the task we stored off when we started this task
+				TSharedPtr<IAnalyticsSpan>* SpanPtr = TaskSpans.Find(TaskGuid);
+
+				if (SpanPtr != nullptr)
+				{
+					TSharedPtr<IAnalyticsSpan> SlowTaskSpan = *SpanPtr;
+
+					FStudioTelemetry::Get().EndSpan(SlowTaskSpan);
+
+					// Remove the SlowTask span from the registry
+					TaskSpans.Remove(TaskGuid);
+				}
+			});
+	}
+
+	FEditorDelegates::OnMapLoad.AddLambda([this](const FString& MapName, FCanLoadMap& OutCanLoadMap)
+		{
+			if (MapName.Len() > 0)
+			{
+				// The Editor loads a new map
+				EditorLoadMapSpan = FStudioTelemetry::Get().StartSpan(EditorLoadMapSpanName, EditorSpan);
+			}
+		});
+
+	FEditorDelegates::OnMapOpened.AddLambda([this](const FString& MapName, bool Unused)
+		{
+			if (EditorLoadMapSpan.IsValid())
+			{
+				// The new editor map was actually opened
+				EditorMapName = FPaths::GetBaseFilename(MapName);
+
+				TArray<FAnalyticsEventAttribute> Attributes;
+				Attributes.Emplace(TEXT("MapName"), EditorMapName);
+
+				EditorSpan->AddAttributes(Attributes);
+				EditorLoadMapSpan->AddAttributes(Attributes);
+
+				FStudioTelemetry::Get().EndSpan(EditorLoadMapSpan);
+				
+				FEditorTelemetry::RecordEvent_Loading(TEXT("LoadMap"), EditorLoadMapSpan->GetDuration(), EditorLoadMapSpan->GetAttributes());
+				FEditorTelemetry::RecordEvent_CoreSystems(TEXT("LoadMap"), EditorLoadMapSpan->GetAttributes());
+
+				EditorLoadMapSpan.Reset();
+			}	
+		});
+
+	// Install any plugin load/unload callbacks
+	FModuleManager::Get().OnModulesChanged().AddLambda([this](FName ModuleName, EModuleChangeReason ChangeReason)
+		{
+			switch ( ChangeReason )
+			{ 
+				default:
+				{
+					break;
+				}
+				
+				case EModuleChangeReason::ModuleLoaded:
+				{
+					TotalPluginCount++;
+
+					// Hook into Asset Registry Scan callbacks as as soon as it is loaded
+					if (ModuleName == TEXT("AssetRegistry"))
+					{
+						FAssetRegistryModule& AssetRegistryModule = FModuleManager::LoadModuleChecked<FAssetRegistryModule>(TEXT("AssetRegistry"));
+
+						AssetRegistryModule.Get().OnScanStarted().AddLambda([this]()
+							{
+								if (AssetRegistryScanCount == 0)
+								{
+									// Start the Asset Registry Scan span
+									AssetRegistryScanSpan = FStudioTelemetry::Get().StartSpan(AssetRegistryScanSpanName, EditorSpan);
+								}	
+
+								AssetRegistryScanCount++;
+							});
+
+						AssetRegistryModule.Get().OnScanEnded().AddLambda([this]()
+							{
+								AssetRegistryScanCount--;
+
+								if (AssetRegistryScanCount == 0)
+								{
+									// End the Asset Registry Scan span
+									TArray<FAnalyticsEventAttribute> Attributes;
+									Attributes.Emplace(TEXT("MapName"), EditorMapName);
+									FStudioTelemetry::Get().EndSpan(AssetRegistryScanSpan, Attributes);
+								}	
+							});
+					}
+
+					break;
+				}
+		
+				case EModuleChangeReason::ModuleUnloaded:
+				{
+					TotalPluginCount--;
+					break;
+				}	
+			}
+		});
+
+	FEditorDelegates::OnEditorBoot.AddLambda([this](double)
+		{	
+			TArray<FAnalyticsEventAttribute> Attributes;
+			Attributes.Emplace(TEXT("PluginCount"), TotalPluginCount);
+			EditorBootSpan->AddAttributes(Attributes);
+
+			FStudioTelemetry::Get().EndSpan(EditorBootSpan);
+
+			// Callback is received when the editor has booted but has not been initialized
+			FEditorTelemetry::RecordEvent_Loading(TEXT("BootEditor"), EditorBootSpan->GetDuration(), EditorBootSpan->GetAttributes());
+			FEditorTelemetry::RecordEvent_CoreSystems(TEXT("BootEditor"), EditorBootSpan->GetAttributes());
+
+			EditorInitilizeSpan = FStudioTelemetry::Get().StartSpan(EditorInitilizeSpanName, EditorSpan);
+		});
+
+	FEditorDelegates::OnEditorInitialized.AddLambda([this](double TimeToInitializeEditor)
+		{
+			TimeToBootEditor = TimeToInitializeEditor;
+
+			// Editor has initialized
+			TArray<FAnalyticsEventAttribute> Attributes;
+			Attributes.Emplace(TEXT("MapName"), EditorMapName);
+			Attributes.Emplace(TEXT("PluginCount"), TotalPluginCount);
+
+			// Editor has finished initializing so start the Editor Interact span
+			FStudioTelemetry::Get().EndSpan(EditorInitilizeSpan);
+			EditorInteractSpan = FStudioTelemetry::Get().StartSpan(EditorInteractSpanName, EditorSpan);
+						
+			FEditorTelemetry::RecordEvent_Loading(TEXT("TotalEditorStartup"), TimeToInitializeEditor, Attributes);
+			FEditorTelemetry::RecordEvent_CoreSystems(TEXT("TotalEditorStartup"), Attributes);
+
+			ensureMsgf(GEditor, TEXT("GEditor was not valid"));
+
+			if (GEditor != nullptr)
+			{
+				// Install callbacks for Open Asset Dialogue
+				GEditor->GetEditorSubsystem<UAssetEditorSubsystem>()->OnAssetEditorRequestedOpen().AddLambda([this](UObject* Asset)
+					{
+						AssetOpenStartTime = FPlatformTime::Seconds();
+						FStudioTelemetry::Get().StartSpan(OpenAssetEditorSpan);
+					});
+
+				GEditor->GetEditorSubsystem<UAssetEditorSubsystem>()->OnAssetOpenedInEditor().AddLambda([this](UObject* Asset, IAssetEditorInstance*)
+					{
+						TArray<FAnalyticsEventAttribute> Attributes;
+						Attributes.Emplace(TEXT("MapName"), EditorMapName);
+						
+						if (Asset != nullptr)
+						{
+							Attributes.Emplace(TEXT("AssetPath"), Asset->GetFullName());
+							Attributes.Emplace(TEXT("AssetClass"), Asset->GetClass()->GetName());
+
+							FEditorTelemetry::RecordEvent_Loading(TEXT("OpenAssetEditor"), FPlatformTime::Seconds() - AssetOpenStartTime, Attributes);
+						}
+
+						FStudioTelemetry::Get().EndSpan(OpenAssetEditorSpan, Attributes);			
+				});
+
+				// Setup a timer for a Heartbeat callback
+				FTimerDelegate HeartbeatDelegate;
+				HeartbeatDelegate.BindRaw(this, &FEditorTelemetry::HeartbeatCallback);
+				GEditor->GetTimerManager()->SetTimer(TelemetryHeartbeatTimerHandle, HeartbeatDelegate, HeartbeatIntervalSeconds, true);
+
+				// Setup the timer for the Hitch Detector callback
+				if (IsRunningCommandlet() == false)
+				{
+					FTimerDelegate HitchSamplerDelegate;
+					HitchSamplerDelegate.BindRaw(this, &FEditorTelemetry::HitchSamplerCallback);
+					GEditor->GetTimerManager()->SetTimer(TelemetryHitchSamplerTimerHandle, HitchSamplerDelegate, HitchSamplerIntervalSeconds, true);
+				}
+			}
+
+			ensureMsgf(GUnrealEd, TEXT("GUnrealEd was not valid"));
+			
+			if (GUnrealEd!=nullptr && GUnrealEd->CookServer != nullptr )
+			{
+				UE::Cook::FDelegates::CookByTheBookFinished.AddLambda([this](UE::Cook::ICookInfo& CookInfo)
+					{
+						TArray<FAnalyticsEventAttribute> Attributes;
+						Attributes.Emplace(TEXT("MapName"), EditorMapName);
+						
+						FEditorTelemetry::RecordEvent_Cooking(Attributes);
+						FEditorTelemetry::RecordEvent_CoreSystems(TEXT("Cooking"), Attributes);
+					});
+			}
+		});
+
+	// Install PIE Mode callbacks
+	FEditorDelegates::StartPIE.AddLambda([this](bool)
+		{
+			// PIE mode has been started. The user has pressed the Start PIE button.
+			// Finish the Editor span
+			FStudioTelemetry::Get().EndSpan(EditorSpan);
+			EditorSessionCount++;
+
+			// Start PIE span
+			PIESpan = FStudioTelemetry::Get().StartSpan(PIESpanName);
+
+			// Append the PIE transition count to the PIE name
+			PIEStartupSpan = FStudioTelemetry::Get().StartSpan(PIESessionCount == 0? PIEStartupSpanName : FName(*FString::Printf(TEXT("%s%d"), *PIEStartupSpanName.ToString(), PIESessionCount)), PIESpan);
+		
+			TArray<FAnalyticsEventAttribute> Attributes;
+			Attributes.Emplace(TEXT("MapName"), EditorMapName);
+
+			PIESpan->AddAttributes(Attributes);
+			PIEStartupSpan->AddAttributes(Attributes);
+		});
+
+	FEditorDelegates::PreBeginPIE.AddLambda([this](bool)
+		{	
+			PIEPreBeginSpan = FStudioTelemetry::Get().StartSpan(PIESessionCount == 0 ? PIEPreBeginSpanName : FName(*FString::Printf(TEXT("%s%d"), *PIEPreBeginSpanName.ToString(), PIESessionCount)), PIESpan);
+			PIEPreBeginSpan->AddAttributes(PIESpan->GetAttributes());
+		});
+
+	FEditorDelegates::BeginPIE.AddLambda([this](bool)
+		{				
+			FStudioTelemetry::Get().EndSpan(PIEPreBeginSpan);
+		});
+
+	FWorldDelegates::OnPIEMapCreated.AddLambda([this](UGameInstance* GameInstance)
+		{
+			// A new PIE map was created
+			PIELoadMapSpan = FStudioTelemetry::Get().StartSpan(PIELoadMapSpanName, PIEStartupSpan);
+		});
+
+	FWorldDelegates::OnPIEMapReady.AddLambda([this](UGameInstance* GameInstance)
+		{
+			// PIE map is now loaded and ready to use
+			PIEMapName = FPaths::GetBaseFilename(GameInstance->PIEMapName);
+
+			TArray<FAnalyticsEventAttribute> Attributes;
+			Attributes.Emplace(TEXT("MapName"), EditorMapName);
+			Attributes.Emplace(TEXT("PIE_MapName"), PIEMapName);
+
+			PIELoadMapSpan->AddAttributes(Attributes);
+			
+			FStudioTelemetry::Get().EndSpan(PIELoadMapSpan);
+
+			FEditorTelemetry::RecordEvent_Loading(TEXT("PIE.LoadMapTime"), PIELoadMapSpan->GetDuration(), PIELoadMapSpan->GetAttributes());
+			FEditorTelemetry::RecordEvent_CoreSystems(TEXT("PIE.LoadMapTime"), PIELoadMapSpan->GetAttributes());
+		});
+
+	FWorldDelegates::OnPIEReady.AddLambda([this](UGameInstance* GameInstance)
+		{
+			if (PIESpan.IsValid())
+			{
+				if ( PIEStartupSpan.IsValid() )
+				{
+					// PIE is now ready for user interaction
+					
+					// Keep track of the PIE transition counts
+					TArray<FAnalyticsEventAttribute> Attributes;
+					Attributes.Emplace(TEXT("PIE_TransitionCount"), PIESessionCount );
+					
+					PIESpan->AddAttributes(Attributes);
+					PIEStartupSpan->AddAttributes(Attributes);
+						
+					FStudioTelemetry::Get().EndSpan(PIEStartupSpan);
+
+					// Record the time from start PIE to PIE
+					FEditorTelemetry::RecordEvent_Loading(TEXT("PIE.TotalStartupTime"), PIEStartupSpan->GetDuration(), PIEStartupSpan->GetAttributes());
+					FEditorTelemetry::RecordEvent_CoreSystems(TEXT("PIE.TotalStartupTime"), PIEStartupSpan->GetAttributes());
+
+					if (PIESessionCount == 0)
+					{
+						const double TimeInEditor = EditorLoadMapSpan.IsValid() ? EditorLoadMapSpan->GetDuration() : 0.0;
+						const double TimeToStartPIE = PIEStartupSpan->GetDuration();
+						const double TimeToBootToPIE = TimeToBootEditor + TimeInEditor + TimeToStartPIE;
+
+						// Record the absolute time from editor boot to PIE
+						FEditorTelemetry::RecordEvent_Loading(TEXT("TimeToPIE"), TimeToBootToPIE, PIEStartupSpan->GetAttributes());
+						FEditorTelemetry::RecordEvent_CoreSystems(TEXT("TimeToPIE"), PIEStartupSpan->GetAttributes());
+					}
+				}
+
+				PIEInteractSpan = FStudioTelemetry::Get().StartSpan( PIESessionCount==0? PIEInteractSpanName : FName(*FString::Printf(TEXT("%s%d"), *PIEInteractSpanName.ToString(), PIESessionCount)), PIESpan);
+				PIEInteractSpan->AddAttributes(PIESpan->GetAttributes());
+			}
+		});
+
+	FEditorDelegates::EndPIE.AddLambda([this](bool)
+		{
+			if (PIESpan.IsValid())
+			{
+				// PIE is ending so no longer interactive
+				FStudioTelemetry::Get().EndSpan(PIEInteractSpan);
+				PIEShutdownSpan = FStudioTelemetry::Get().StartSpan(PIESessionCount == 0 ? PIEShutdownSpanName : FName(*FString::Printf(TEXT("%s%d"), *PIEShutdownSpanName.ToString(), PIESessionCount)), PIESpan);
+				PIEShutdownSpan->AddAttributes(PIESpan->GetAttributes());
+			}
+		});
+
+	FEditorDelegates::ShutdownPIE.AddLambda([this](bool)
+		{
+			if (PIESpan.IsValid())
+			{
+				// PIE has shutdown, ie. the user has pressed the Stop PIE button, and we are going back to interactive Editor mode	
+				FStudioTelemetry::Get().EndSpan(PIESpan);
+
+				FEditorTelemetry::RecordEvent_Loading(TEXT("PIE.EndTime"), PIESpan->GetDuration(), PIESpan->GetAttributes());
+				FEditorTelemetry::RecordEvent_CoreSystems(TEXT("PIE.EndTime"), PIESpan->GetAttributes());
+			}
+
+			PIESessionCount++;
+
+			TArray<FAnalyticsEventAttribute> Attributes;
+			Attributes.Emplace(TEXT("MapName"), EditorMapName);
+
+			// Restart the Editor span
+			EditorSpan = FStudioTelemetry::Get().StartSpan(EditorSpanName, Attributes);
+			EditorInteractSpan = FStudioTelemetry::Get().StartSpan(FName(*FString::Printf(TEXT("%s%d"), *EditorInteractSpanName.ToString(), EditorSessionCount)), EditorSpan);
+		});
+
+	
+	// Install Cooking Callbacks
+	UE::Cook::FDelegates::CookByTheBookStarted.AddLambda([this](UE::Cook::ICookInfo& CookInfo)
+	{
+		// Begin the cooking span	
+		CookingSpan = FStudioTelemetry::Get().StartSpan(TEXT("Cooking"));
+
+		TArray<FAnalyticsEventAttribute> Attributes;
+		Attributes.Emplace(TEXT("MapName"), EditorMapName);
+
+		CookingSpan->AddAttributes(Attributes);
+	});
+
+	UE::Cook::FDelegates::CookByTheBookFinished.AddLambda([this](UE::Cook::ICookInfo& CookInfo)
+	{
+		// End the cooking span
+	
+		// Suppress sending telemetry from CookWorkers for now.
+		uint32 MultiprocessId = 0;
+		FParse::Value(FCommandLine::Get(), TEXT("-MultiprocessId="), MultiprocessId);
+		if (MultiprocessId != 0)
+		{
+			return;
+		}
+
+		FEditorTelemetry::RecordEvent_Cooking(CookingSpan->GetAttributes());
+		FEditorTelemetry::RecordEvent_CoreSystems(TEXT("Cooking"), CookingSpan->GetAttributes());
+
+		FStudioTelemetry::Get().EndSpan(CookingSpan);
+	});
+
+	// Install Content Browser callbacks
+	FContentBrowserModule* ContentBrowserModule = FModuleManager::GetModulePtr<FContentBrowserModule>( TEXT("ContentBrowser") );
+	
+	FTelemetryRouter& Router = FTelemetryRouter::Get();
+	{
+		using namespace UE::Telemetry::ContentBrowser;
+		Router.OnTelemetry<FBackendFilterTelemetry>([this](const FBackendFilterTelemetry& Data)
+		{
+			FString DataFilterText = LexToString(FJsonNull{});
+			if (Data.DataFilter)
+			{
+				Private::FAnalyticsJsonWriter J(&DataFilterText);
+				J.WriteObjectStart();
+				J.WriteValue("RecursivePaths", Data.DataFilter->bRecursivePaths);
+				J.WriteValue("ItemTypeFilter", UEnum::GetValueOrBitfieldAsString(Data.DataFilter->ItemTypeFilter));
+				J.WriteValue("ItemCategoryFilter", UEnum::GetValueOrBitfieldAsString(Data.DataFilter->ItemCategoryFilter));
+				J.WriteValue("ItemAttributeFilter", UEnum::GetValueOrBitfieldAsString(Data.DataFilter->ItemAttributeFilter));
+				TArray<const UScriptStruct*> FilterTypes = Data.DataFilter->ExtraFilters.GetFilterTypes();
+				if (FilterTypes.Num() > 0)
+				{
+					J.WriteArrayStart("FilterTypes");
+					for (const UScriptStruct* Type : FilterTypes)
+					{
+						J.WriteValue(Type->GetPathName());
+					}
+					J.WriteArrayEnd();				
+				}
+				J.WriteObjectEnd();
+				J.Close();
+			}
+			
+			FStudioTelemetry::Get().RecordEvent(TEXT("Editor.AssetView.BackendFilter"), 
+			{
+				{ TEXT("SchemaVersion"), 1 },
+				{ TEXT("ViewCorrelationGuid"), Data.ViewCorrelationGuid },
+				{ TEXT("FilterSessionCorrelationGuid"), Data.FilterSessionCorrelationGuid },
+				{ TEXT("HasCustomItemSources"), Data.bHasCustomItemSources },
+				{ TEXT("RefreshSourceItemsDurationSeconds"), Data.RefreshSourceItemsDurationSeconds },
+				{ TEXT("NumBackendItems"), Data.NumBackendItems },
+				{ TEXT("DataFilter"), FJsonFragment(MoveTemp(DataFilterText)) },
+			});
+		});
+
+		Router.OnTelemetry<FFrontendFilterTelemetry>([this](const FFrontendFilterTelemetry& Data)
+		{
+			FString FilterText = LexToString(FJsonNull{});
+			if (Data.FrontendFilters.IsValid() && Data.FrontendFilters->Num())
+			{
+				Private::FAnalyticsJsonWriter J(&FilterText);
+				J.WriteArrayStart();
+				for (int32 i=0; i < Data.FrontendFilters->Num(); ++i)
+				{
+					TSharedPtr<IFilter<FAssetFilterType>> Filter = Data.FrontendFilters->GetFilterAtIndex(i);
+					J.WriteValue(Filter->GetName());
+				}
+				J.WriteArrayEnd();
+				J.Close();
+			}
+			FStudioTelemetry::Get().RecordEvent(TEXT("Editor.AssetView.FrontendFilter"), 
+			{
+				{ TEXT("SchemaVersion"), 1 },
+				{ TEXT("ViewCorrelationGuid"), Data.ViewCorrelationGuid },
+				{ TEXT("FilterSessionCorrelationGuid"), Data.FilterSessionCorrelationGuid },
+				{ TEXT("TotalItemsToFilter"), Data.TotalItemsToFilter },
+				{ TEXT("PriorityItemsToFilter"), Data.PriorityItemsToFilter },
+				{ TEXT("TotalResults"), Data.TotalResults },
+				{ TEXT("AmortizeDurationSeconds"), Data.AmortizeDuration },
+				{ TEXT("WorkDurationSeconds"), Data.WorkDuration },
+				{ TEXT("ResultLatency"), AnalyticsOptionalToStringOrNull(Data.ResultLatency) },
+				{ TEXT("TimeUntilInteractionSeconds"), AnalyticsOptionalToStringOrNull(Data.TimeUntilInteraction) },
+				{ TEXT("Completed"), Data.bCompleted },
+				{ TEXT("FrontendFilters"), FJsonFragment(MoveTemp(FilterText)) },
+			});
+		});
+
+		RegisterCollectionWorkflowDelegates(Router);
+	}
+	{
+		using namespace UE::Telemetry::AssetRegistry;
+
+		Router.OnTelemetry<FStartupTelemetry>([this](const FStartupTelemetry& Data){
+			FStudioTelemetry::Get().RecordEvent(TEXT("Editor.AssetRegistry.Startup"), 
+			{
+				{ TEXT("SchemaVersion"), 1 },
+				{ TEXT("Duration"), Data.StartupDuration },
+				{ TEXT("StartedAsyncGather"), Data.bStartedAsyncGather },
+			});
+		});
+		Router.OnTelemetry<FSynchronousScanTelemetry>([this](const FSynchronousScanTelemetry& Data){
+			if (Data.Duration < 0.5)
+			{
+				return;
+			}
+			FString DirectoriesText;
+			{
+				Private::FAnalyticsJsonWriter J(&DirectoriesText);
+				J.WriteArrayStart();
+				for (const FString& Directory : MakeArrayView(Data.Directories).Left(100))
+				{
+					J.WriteValue(Directory);
+				}
+				J.WriteArrayEnd();
+				J.Close();
+			}
+			FString FilesText;
+			{
+				Private::FAnalyticsJsonWriter J(&FilesText);
+				J.WriteArrayStart();
+				for (const FString& File : MakeArrayView(Data.Files).Left(100))
+				{
+					J.WriteValue(File);
+				}
+				J.WriteArrayEnd();
+				J.Close();
+			}
+			FStudioTelemetry::Get().RecordEvent(TEXT("Editor.AssetRegistry.SynchronousScan"), 
+			{
+				{ TEXT("SchemaVersion"), 1 },
+				{ TEXT("Directories"), FJsonFragment(MoveTemp(DirectoriesText)) },
+				{ TEXT("Files"), FJsonFragment(MoveTemp(FilesText)) },
+				{ TEXT("Flags"), LexToString(Data.Flags) },
+				{ TEXT("NumFoundAssets"), Data.NumFoundAssets },
+				{ TEXT("DurationSeconds"), Data.Duration },
+				{ TEXT("InitialSearchStarted"), Data.bInitialSearchStarted },
+				{ TEXT("InitialSearchCompleted"), Data.bInitialSearchCompleted },
+			});
+		});
+		Router.OnTelemetry<FGatherTelemetry>([this](const FGatherTelemetry& Data){
+			FStudioTelemetry::Get().RecordEvent(TEXT("Editor.AssetRegistry.InitialScan"), 
+			{
+				{ TEXT("SchemaVersion"), 1 },
+				{ TEXT("TotalDurationSeconds"), Data.TotalSearchDurationSeconds },
+				{ TEXT("TotalWorkSeconds"), Data.TotalWorkTimeSeconds },
+				{ TEXT("DiscoverySeconds"), Data.DiscoveryTimeSeconds },
+				{ TEXT("GatherSeconds"), Data.GatherTimeSeconds },
+				{ TEXT("StoreSeconds"), Data.StoreTimeSeconds },
+				{ TEXT("NumCachedDirectories"), Data.NumCachedDirectories },
+				{ TEXT("NumUncachedDirectories"), Data.NumUncachedDirectories },
+				{ TEXT("NumCachedAssetFiles"), Data.NumCachedAssetFiles },
+				{ TEXT("NumUncachedAssetFiles"), Data.NumUncachedAssetFiles },
+			});
+		});
+		Router.OnTelemetry<FDirectoryWatcherUpdateTelemetry>([this](const FDirectoryWatcherUpdateTelemetry& Data){
+			if (Data.DurationSeconds < 0.5)
+			{
+				return;
+			}
+			FStudioTelemetry::Get().RecordEvent(TEXT("Editor.AssetRegistry.DirectoryWatcherUpdate"), 
+			{
+				{ TEXT("SchemaVersion"), 1 },
+				{ TEXT("NumChanges"), Data.Changes.Num() },
+				{ TEXT("DurationSeconds"), Data.DurationSeconds },
+				{ TEXT("InitialSearchStarted"), Data.bInitialSearchStarted },
+				{ TEXT("InitialSearchCompleted"), Data.bInitialSearchCompleted },
+			});
+		});
+	}
+
+	{
+		UE::Virtualization::GetAnalyticsRecordEvent().AddLambda([](const FString& EventName, const TArray<FAnalyticsEventAttribute>& Attributes, UE::Virtualization::EAnalyticsFlags Flags)
+			{
+				FStudioTelemetry::Get().RecordEvent(EventName, Attributes);
+
+				if (EnumHasAllFlags(Flags, UE::Virtualization::EAnalyticsFlags::Flush))
+				{
+					if (TSharedPtr<IAnalyticsProvider> Provider = FStudioTelemetry::Get().GetProvider().Pin())
+					{
+						Provider->FlushEvents();
+
+						// It is quite likely that one of the analytics providers is sending data via the FHttpManager so we
+						// need to flush that as well to make sure that the message is sent. If the virtualization system has
+						// requested that we flush then it is quite likely that the process is to be terminated in which case
+						// not flushing the FHttpManager will result in analytics not being properly recorded.
+						if (FHttpModule* HttpModule = FModuleManager::GetModulePtr<FHttpModule>("HTTP"))
+						{
+							FHttpManager& HttpManager = HttpModule->GetHttpManager();
+							HttpManager.Flush(EHttpFlushReason::FullFlush);
+						}	
+					}
+				}
+			});
+	}
+}
+
+void FEditorTelemetry::EndSession()
+{
+	FStudioTelemetry::Get().EndSpan(EditorSpan);
+}
+
+UE_ENABLE_OPTIMIZATION_SHIP
diff --git a/Engine/Plugins/Experimental/EditorTelemetry/Source/Private/EditorTelemetryModule.cpp b/Engine/Plugins/Experimental/EditorTelemetry/Source/Private/EditorTelemetryModule.cpp
new file mode 100644
--- /dev/null
+++ b/Engine/Plugins/Experimental/EditorTelemetry/Source/Private/EditorTelemetryModule.cpp
@@ -0,0 +1,28 @@
+// Copyright Epic Games, Inc. All Rights Reserved.
+
+#include "EditorTelemetryModule.h"
+#include "EditorTelemetry.h"
+#include "StudioTelemetry.h"
+#include "Modules/ModuleManager.h"
+
+IMPLEMENT_MODULE(FEditorTelemetryModule, EditorTelemetry);
+
+void FEditorTelemetryModule::StartupModule()
+{
+	FStudioTelemetry::Get().GetOnStartSession().AddLambda([this]()
+		{
+			FEditorTelemetry::Get().StartSession();
+		}
+	);
+	
+	FStudioTelemetry::Get().GetOnEndSession().AddLambda([this]()
+		{
+			FEditorTelemetry::Get().EndSession();
+		}
+	);
+}
+
+void FEditorTelemetryModule::ShutdownModule()
+{
+	FEditorTelemetry::Get().EndSession();
+}
diff --git a/Engine/Plugins/Experimental/EditorTelemetry/Source/Public/EditorTelemetry.h b/Engine/Plugins/Experimental/EditorTelemetry/Source/Public/EditorTelemetry.h
new file mode 100644
--- /dev/null
+++ b/Engine/Plugins/Experimental/EditorTelemetry/Source/Public/EditorTelemetry.h
@@ -0,0 +1,98 @@
+// Copyright Epic Games, Inc. All Rights Reserved.
+
+#pragma once
+
+#include "Interfaces/IAnalyticsProvider.h"
+#include "Interfaces/IAnalyticsTracer.h"
+#include "Containers/Array.h"
+#include "Containers/Map.h"
+#include "Containers/UnrealString.h"
+#include "CoreMinimal.h"
+#include "Templates/SharedPointer.h"
+#include "Engine/EngineTypes.h"
+
+class FTelemetryRouter;
+struct FTimerHandle;
+
+/**
+ * A class that implements a variety of pre-configured Core and Editor telemetry events that can be used to evaluate the efficiency of the most common developer workflows
+ */
+class EDITORTELEMETRY_API FEditorTelemetry
+{
+public:
+	FEditorTelemetry() {};
+	~FEditorTelemetry() {};
+
+	static FEditorTelemetry& Get();
+
+	void StartSession();
+	void EndSession();
+
+private:	
+
+	/** Various event handling functions */
+	static void RecordEvent_Cooking(TArray<FAnalyticsEventAttribute> Attributes = {});
+	static void RecordEvent_Loading(const FString& LoadingName, double LoadingSeconds, TArray<FAnalyticsEventAttribute> Attributes = {});
+	static void RecordEvent_CoreSystems(const FString& Context, TArray<FAnalyticsEventAttribute> Attributes = {});
+	static void RecordEvent_DDCResource(const FString& Context, TArray<FAnalyticsEventAttribute> Attributes = {});
+	static void RecordEvent_DDCSummary(const FString& Context, TArray<FAnalyticsEventAttribute> Attributes = {});
+	static void RecordEvent_IAS(const FString& Context, TArray<FAnalyticsEventAttribute> Attributes = {});
+	static void RecordEvent_Zen(const FString& Context, TArray<FAnalyticsEventAttribute> Attributes = {});
+	static void RecordEvent_VirtualAssets(const FString& Context, TArray<FAnalyticsEventAttribute> Attributes = {});
+	static void RegisterCollectionWorkflowDelegates(FTelemetryRouter& Router);
+	void HitchSamplerCallback();
+	void HeartbeatCallback();
+	
+	TSharedPtr<IAnalyticsSpan> EditorSpan;
+	TSharedPtr<IAnalyticsSpan> EditorBootSpan;
+	TSharedPtr<IAnalyticsSpan> EditorInteractSpan;
+	TSharedPtr<IAnalyticsSpan> EditorInitilizeSpan;
+	TSharedPtr<IAnalyticsSpan> EditorLoadMapSpan;
+	TSharedPtr<IAnalyticsSpan> PIESpan;
+	TSharedPtr<IAnalyticsSpan> PIEPreBeginSpan;
+	TSharedPtr<IAnalyticsSpan> PIEStartupSpan;
+	TSharedPtr<IAnalyticsSpan> PIELoadMapSpan;
+	TSharedPtr<IAnalyticsSpan> PIEInteractSpan;		
+	TSharedPtr<IAnalyticsSpan> PIEShutdownSpan;
+	TSharedPtr<IAnalyticsSpan> CookingSpan;
+	TSharedPtr<IAnalyticsSpan> HitchingSpan;
+	TSharedPtr<IAnalyticsSpan> AssetRegistryScanSpan;
+	
+	const FName EditorSpanName = TEXT("Editor");
+	const FName EditorBootSpanName = TEXT("Editor.Boot");
+	const FName EditorInitilizeSpanName = TEXT("Editor.Initialize");
+	const FName EditorInteractSpanName = TEXT("Editor.Interact");
+	const FName EditorLoadMapSpanName = TEXT("Editor.LoadMap");
+	const FName AssetRegistryScanSpanName = TEXT("Editor.AssetRegistryScan");
+	const FName PIESpanName = TEXT("PIE");
+	const FName PIEStartupSpanName = TEXT("PIE.Startup");
+	const FName PIEPreBeginSpanName = TEXT("PIE.PreBegin");
+	const FName PIELoadMapSpanName = TEXT("PIE.LoadMap");
+	const FName PIEInteractSpanName = TEXT("PIE.Interact");
+	const FName PIEShutdownSpanName = TEXT("PIE.Shutdown");
+	const FName CookingSpanName = TEXT("Cooking");
+	const FName HitchingSpanName = TEXT("Hitching");
+	const FName OpenAssetEditorSpan = TEXT("Open Asset Editor");
+	const float HeartbeatIntervalSeconds = 1.0;
+	const float HitchSamplerIntervalSeconds = 0.1;
+	const float MinFPSForHitching = 5.0;
+
+	TMap<FGuid, TSharedPtr<IAnalyticsSpan>> TaskSpans;
+	FCriticalSection TaskSpanCriticalSection;
+
+	FTimerHandle TelemetryHeartbeatTimerHandle;
+	FTimerHandle TelemetryHitchSamplerTimerHandle;
+	FString EditorMapName;
+	FString PIEMapName;
+	uint32 EditorSessionCount = 0;
+	uint32 PIESessionCount = 0;
+	double SessionStartTime;
+	double AssetOpenStartTime;
+	double TimeToBootEditor;
+	double HitchAvergageFPS = 0;
+	uint32 HitchSampleCount = 0;
+	uint32 TotalPluginCount = 0;
+	uint32 AssetRegistryScanCount = 0;
+};
+
+
diff --git a/Engine/Plugins/Experimental/EditorTelemetry/Source/Public/EditorTelemetryModule.h b/Engine/Plugins/Experimental/EditorTelemetry/Source/Public/EditorTelemetryModule.h
new file mode 100644
--- /dev/null
+++ b/Engine/Plugins/Experimental/EditorTelemetry/Source/Public/EditorTelemetryModule.h
@@ -0,0 +1,16 @@
+// Copyright Epic Games, Inc. All Rights Reserved.
+
+#pragma once
+
+#include "CoreMinimal.h"
+#include "Modules/ModuleInterface.h"
+
+class EDITORTELEMETRY_API FEditorTelemetryModule : public IModuleInterface
+{
+public:
+	// IModuleInterface
+	virtual void StartupModule() override;
+	virtual void ShutdownModule() override;
+
+private:
+};
diff --git a/Engine/Plugins/Experimental/StudioTelemetry/Source/StudioTelemetry/Private/StudioTelemetry.cpp b/Engine/Plugins/Experimental/StudioTelemetry/Source/StudioTelemetry/Private/StudioTelemetry.cpp
--- a/Engine/Plugins/Experimental/StudioTelemetry/Source/StudioTelemetry/Private/StudioTelemetry.cpp
+++ b/Engine/Plugins/Experimental/StudioTelemetry/Source/StudioTelemetry/Private/StudioTelemetry.cpp
@@ -4,7 +4,6 @@
 #include "StudioTelemetryLog.h"
 
 #if WITH_EDITOR
-#include "StudioTelemetryEditor.h"
 #include "Horde.h"
 #endif
 
@@ -35,7 +34,7 @@ FStudioTelemetry& FStudioTelemetry::Get()
 	return StudioTelemetryInstance;
 }
 
-void FStudioTelemetry::SetRecordEventCallback(OnRecordEvent Callback )
+void FStudioTelemetry::SetRecordEventCallback(OnRecordEventCallback Callback )
 {
 	RecordEventCallback = Callback;
 
@@ -50,25 +49,15 @@ void FStudioTelemetry::StartupModule()
 {
 	UE_LOG(LogStudioTelemetry, Display, TEXT("Starting StudioTelemetry Module"));
 
-	// Load the configureation
+	// Load the configuration
 	FStudioTelemetry::Get().LoadConfiguration();
 
 	// Create the provider and start the analytics session
 	FStudioTelemetry::Get().StartSession();
-
-#if WITH_EDITOR
-	// Initialize the analytics subsystems 
-	FStudioTelemetryEditor::Get().Initialize();
-#endif
 }
 
 void FStudioTelemetry::ShutdownModule()
 {
-#if WITH_EDITOR
-	// Shutdown the analytics subsystems
-	FStudioTelemetryEditor::Get().Shutdown();
-#endif
-
 	// End the session and destroy analytics provider
 	FStudioTelemetry::Get().EndSession();
 
@@ -77,6 +66,8 @@ void FStudioTelemetry::ShutdownModule()
 
 void FStudioTelemetry::EndSession()
 {
+	OnEndSession.Broadcast();
+
 	// End session for the tracer and the provider
 	if (AnalyticsTracer.IsValid())
 	{
@@ -97,7 +88,7 @@ void FStudioTelemetry::LoadConfiguration()
 {
 	const FString TelemetryConfigurationSection(TEXT("StudioTelemetry.Config"));
 
-	// Look for the configuration seetings in the Engine.ini files
+	// Look for the configuration settings in the Engine.ini files
 	TArray<FString> SectionNames;
 
 	if (GConfig->GetSectionNames(GEngineIni, SectionNames))
@@ -242,6 +233,8 @@ void FStudioTelemetry::StartSession()
 		// Bind the pre-exit callback
 		FCoreDelegates::OnEnginePreExit.AddRaw(&FStudioTelemetry::Get(), &FStudioTelemetry::EndSession);
 
+		OnStartSession.Broadcast();
+
 		UE_LOG(LogStudioTelemetry, Log, TEXT("Started StudioTelemetry Session"));
 	}
 }
@@ -253,6 +246,8 @@ void FStudioTelemetry::RecordEvent(const FString& EventName, const TArray<FAnaly
 		FScopeLock ScopeLock(&CriticalSection);
 		AnalyticsProvider->RecordEvent(EventName, Attributes);
 	}
+
+	OnRecordEvent.Broadcast(EventName, Attributes);
 }
 
 void FStudioTelemetry::RecordEvent(const FName CategoryName, const FString& EventName, const TArray<FAnalyticsEventAttribute>& Attributes)
@@ -262,6 +257,8 @@ void FStudioTelemetry::RecordEvent(const FName CategoryName, const FString& Even
 		FScopeLock ScopeLock(&CriticalSection);
 		AnalyticsProvider->RecordEvent(EventName, Attributes);
 	}
+
+	OnRecordEvent.Broadcast(EventName, Attributes);
 }
 
 void FStudioTelemetry::RecordEventToProvider(const FString& ProviderName, const FString& EventName, const TArray<FAnalyticsEventAttribute>& Attributes)
diff --git a/Engine/Plugins/Experimental/StudioTelemetry/Source/StudioTelemetry/Private/StudioTelemetryEditor.cpp b/Engine/Plugins/Experimental/StudioTelemetry/Source/StudioTelemetry/Private/StudioTelemetryEditor.cpp
deleted file mode 100644
diff --git a/Engine/Plugins/Experimental/StudioTelemetry/Source/StudioTelemetry/Private/StudioTelemetryEditor.h b/Engine/Plugins/Experimental/StudioTelemetry/Source/StudioTelemetry/Private/StudioTelemetryEditor.h
deleted file mode 100644
diff --git a/Engine/Plugins/Experimental/StudioTelemetry/Source/StudioTelemetry/Public/StudioTelemetry.h b/Engine/Plugins/Experimental/StudioTelemetry/Source/StudioTelemetry/Public/StudioTelemetry.h
--- a/Engine/Plugins/Experimental/StudioTelemetry/Source/StudioTelemetry/Public/StudioTelemetry.h
+++ b/Engine/Plugins/Experimental/StudioTelemetry/Source/StudioTelemetry/Public/StudioTelemetry.h
@@ -15,7 +15,7 @@ class FAnalyticsProviderMulticast;
  * Studio Telemetry Plugin API
  * 
  * Notes:
- * Telemetry for Common Editor and Core Engine is collected automatically.
+ * Telemetry for Common Editor and Core Engine is collected automatically via the EditorTelemetry plugin.
  * Telemetry Sessions are started and ended automatically with the plugin initialization and shutdown. As such telemetry will not be captured prior to the plugin initialization.
  * Developers are encouraged to add their own telemetry via this API or to intercept the event recording via the supplied callback on the SetRecordEventCallback API below.
  * It is strongly recommended that developers implement their own IAnalyticsProviderModule where custom recording of telemetry events is desired.
@@ -26,7 +26,7 @@ class FStudioTelemetry : public IModuleInterface
 {
 public:
 
-	typedef TFunction<void(const FString& EventName, const TArray<FAnalyticsEventAttribute>& Attrs)> OnRecordEvent;
+	typedef TFunction<void(const FString& EventName, const TArray<FAnalyticsEventAttribute>& Attrs)> OnRecordEventCallback;
 
 	/** Check whether the module is available*/
 	static STUDIOTELEMETRY_API bool IsAvailable() { return FModuleManager::Get().IsModuleLoaded("StudioTelemetry"); }
@@ -71,8 +71,19 @@ public:
 	STUDIOTELEMETRY_API TSharedPtr<IAnalyticsSpan> GetSessionSpan() const;
 
 	/** Callback for interception of telemetry events recording that can be used by Developers to send telemetry events to their own back end, though it is recommended that Developers implement their own IAnalyticsProvider via their own IAnalyticsProviderModule*/
-	STUDIOTELEMETRY_API void SetRecordEventCallback(OnRecordEvent);
+	STUDIOTELEMETRY_API void SetRecordEventCallback(OnRecordEventCallback);
 
+	/** Delegates for event callbacks **/
+	DECLARE_MULTICAST_DELEGATE(FOnStartSession);
+	STUDIOTELEMETRY_API FOnStartSession& GetOnStartSession() { return OnStartSession; }
+
+	DECLARE_MULTICAST_DELEGATE(FOnEndSession);
+	STUDIOTELEMETRY_API FOnEndSession& GetOnEndSession() { return OnEndSession; }
+
+	DECLARE_MULTICAST_DELEGATE_TwoParams(FOnRecordEvent, const FString&, const TArray<FAnalyticsEventAttribute>&);
+	STUDIOTELEMETRY_API FOnRecordEvent& GetOnRecordEvent() { return OnRecordEvent; }
+
+	/** Scoped Span helper class **/
 	class ScopedSpan
 	{
 	public:
@@ -122,13 +133,15 @@ private:
 	FCriticalSection						CriticalSection;
 	TSharedPtr<FAnalyticsProviderMulticast>	AnalyticsProvider;
 	TSharedPtr<IAnalyticsTracer>			AnalyticsTracer;
-	OnRecordEvent							RecordEventCallback;
+	OnRecordEventCallback					RecordEventCallback;
 	FGuid									SessionGUID;
 	FConfig									Config;
-
-	
+	FOnStartSession							OnStartSession;
+	FOnEndSession							OnEndSession;
+	FOnRecordEvent							OnRecordEvent;
 };
 
+// Useful macros for scoped spans
 #define STUDIO_TELEMETRY_SPAN_SCOPE(Name) FStudioTelemetry::ScopedSpan PREPROCESSOR_JOIN(ScopedSpan, __LINE__)(TEXT(#Name));
 #define STUDIO_TELEMETRY_START_SPAN(Name) if (FStudioTelemetry::Get().IsAvailable()) { FStudioTelemetry::Get().StartSpan(TEXT(#Name));}
 #define STUDIO_TELEMETRY_END_SPAN(Name) if (FStudioTelemetry::Get().IsAvailable()) { FStudioTelemetry::Get().EndSpan(TEXT(#Name));}
diff --git a/Engine/Plugins/Experimental/StudioTelemetry/Source/StudioTelemetry/StudioTelemetry.build.cs b/Engine/Plugins/Experimental/StudioTelemetry/Source/StudioTelemetry/StudioTelemetry.build.cs
--- a/Engine/Plugins/Experimental/StudioTelemetry/Source/StudioTelemetry/StudioTelemetry.build.cs
+++ b/Engine/Plugins/Experimental/StudioTelemetry/Source/StudioTelemetry/StudioTelemetry.build.cs
@@ -12,38 +12,21 @@ public class StudioTelemetry : ModuleRules
 		PublicDependencyModuleNames.AddRange(
 			new string[] {
 				"Core",
-				"HTTP"
+				"HTTP",
+				"Horde"
 			}
 		);
 
 		PrivateDependencyModuleNames.AddRange(
 			new string[] {
 				"CoreUObject",
-				"Engine",
                 "EngineSettings",
                 "BuildSettings",
                 "Analytics",
                 "AnalyticsET",
 				"TelemetryUtils",
-				"RHI"
+				"RHI",
 			}
 		);
-
-		if (Target.bBuildEditor)
-		{
-			PrivateDependencyModuleNames.AddRange(
-				new string[] {
-					"Horde",
-					"AssetRegistry",
-					"UnrealEd",
-					"DerivedDataCache",
-					"Zen",
-					"IoStoreOnDemand",
-					"ContentBrowser",
-					"ContentBrowserData",
-					"TelemetryUtils"
-				}
-			);
-		}
 	}
 }
