Remove IoStoreOnDemand references from the editor telemetry system.

#rnx

diff --git a/Engine/Plugins/Experimental/EditorTelemetry/Source/EditorTelemetry.Build.cs b/Engine/Plugins/Experimental/EditorTelemetry/Source/EditorTelemetry.Build.cs
--- a/Engine/Plugins/Experimental/EditorTelemetry/Source/EditorTelemetry.Build.cs
+++ b/Engine/Plugins/Experimental/EditorTelemetry/Source/EditorTelemetry.Build.cs
@@ -24,7 +24,6 @@ namespace UnrealBuildTool.Rules
 					"StudioTelemetry",
 					"DerivedDataCache",
 					"Zen",
-					"IoStoreOnDemand",
 					"ContentBrowser",
 					"ContentBrowserData",
 				}
diff --git a/Engine/Plugins/Experimental/EditorTelemetry/Source/Private/EditorTelemetry.cpp b/Engine/Plugins/Experimental/EditorTelemetry/Source/Private/EditorTelemetry.cpp
--- a/Engine/Plugins/Experimental/EditorTelemetry/Source/Private/EditorTelemetry.cpp
+++ b/Engine/Plugins/Experimental/EditorTelemetry/Source/Private/EditorTelemetry.cpp
@@ -17,7 +17,6 @@
 #include "FileHelpers.h"
 #include "HttpManager.h"
 #include "HttpModule.h"
-#include "IO/IoStoreOnDemand.h"
 #include "Misc/FeedbackContext.h"
 #include "Modules/ModuleManager.h"
 #include "Policies/CondensedJsonPrintPolicy.h"
@@ -277,27 +276,6 @@ void FEditorTelemetry::RecordEvent_DDCSummary(const FString& Context, TArray<FAn
 #endif			
 }
 
-void FEditorTelemetry::RecordEvent_IAS(const FString& Context, TArray<FAnalyticsEventAttribute> Attributes)
-{
-	// Gather the summary stats
-	FDerivedDataCacheSummaryStats SummaryStats;
-
-	using namespace UE::IoStore;
-	if (FIoStoreOnDemandModule* OnDemandModule = FModuleManager::Get().GetModulePtr<FIoStoreOnDemandModule>("IoStoreOnDemand"))
-	{
-		if (OnDemandModule->IsEnabled())
-		{
-			const int SchemaVersion = 1;
-
-			Attributes.Emplace(TEXT("SchemaVersion"), SchemaVersion);
-			Attributes.Emplace(TEXT("LoadingName"), Context);
-
-			OnDemandModule->ReportAnalytics(Attributes);
-			FStudioTelemetry::Get().RecordEvent(TEXT("Core.IAS"), Attributes);
-		}		
-	}
-}
-
 void FEditorTelemetry::RecordEvent_Zen(const FString& Context, TArray<FAnalyticsEventAttribute> Attributes)
 {
 #if UE_WITH_ZEN
@@ -335,7 +313,6 @@ void FEditorTelemetry::RecordEvent_CoreSystems(const FString& Context, TArray<FA
 {
 	FEditorTelemetry::RecordEvent_DDCResource(Context, Attributes);
 	FEditorTelemetry::RecordEvent_DDCSummary(Context, Attributes);
-	FEditorTelemetry::RecordEvent_IAS(Context, Attributes);
 	FEditorTelemetry::RecordEvent_Zen(Context, Attributes);
 	FEditorTelemetry::RecordEvent_VirtualAssets(Context, Attributes);
 }
