Removed the installed=true flag from Editor Telemetry plugin
Turned off the disable of shipping optimisation MACRO from EditorTelemetry.cpp

diff --git a/Engine/Plugins/Experimental/EditorTelemetry/EditorTelemetry.uplugin b/Engine/Plugins/Experimental/EditorTelemetry/EditorTelemetry.uplugin
--- a/Engine/Plugins/Experimental/EditorTelemetry/EditorTelemetry.uplugin
+++ b/Engine/Plugins/Experimental/EditorTelemetry/EditorTelemetry.uplugin
@@ -2,9 +2,9 @@
 	"FileVersion": 1,
 	"Version": 1,
 	"VersionName": "1.0",
-	"FriendlyName": "Editor Performance",
-	"Description": "Plugin that provides Editor Performance feedback to developers",
-	"Category": "Performance",
+	"FriendlyName": "Editor Telemetry",
+	"Description": "Plugin that emits common telemetry events from the Editor",
+	"Category": "Telemetry",
 	"CreatedBy": "Epic Games, Inc.",
 	"CreatedByURL": "https://epicgames.com",
 	"DocsURL": "",
@@ -13,7 +13,7 @@
 	"EnabledByDefault": true,
 	"CanContainContent": false,
 	"IsBetaVersion": false,
-	"Installed": true,
+	"Installed": false,
 	"Modules": [
 		{
 			"Name": "EditorTelemetry",
diff --git a/Engine/Plugins/Experimental/EditorTelemetry/Source/Private/EditorTelemetry.cpp b/Engine/Plugins/Experimental/EditorTelemetry/Source/Private/EditorTelemetry.cpp
--- a/Engine/Plugins/Experimental/EditorTelemetry/Source/Private/EditorTelemetry.cpp
+++ b/Engine/Plugins/Experimental/EditorTelemetry/Source/Private/EditorTelemetry.cpp
@@ -29,8 +29,6 @@
 #include "UnrealEdGlobals.h"
 #include "Virtualization/VirtualizationSystem.h"
 
-UE_DISABLE_OPTIMIZATION_SHIP
-
 namespace Private
 {
 	const FName ContentBrowserModuleName = TEXT("ContentBrowser");
@@ -1020,5 +1018,3 @@ void FEditorTelemetry::EndSession()
 {
 	FStudioTelemetry::Get().EndSpan(EditorSpan);
 }
-
-UE_ENABLE_OPTIMIZATION_SHIP
