Changed Editor Telemetry plugin so that Editor only telemetry is not recorded when running a commandlet
We will now no longer get Editor Boot, LoadMap or PIE transtion flows recorded in commandlet runs.

diff --git a/Engine/Plugins/Experimental/EditorTelemetry/Source/Private/EditorTelemetry.cpp b/Engine/Plugins/Experimental/EditorTelemetry/Source/Private/EditorTelemetry.cpp
--- a/Engine/Plugins/Experimental/EditorTelemetry/Source/Private/EditorTelemetry.cpp
+++ b/Engine/Plugins/Experimental/EditorTelemetry/Source/Private/EditorTelemetry.cpp
@@ -437,111 +437,250 @@ void FEditorTelemetry::StartSession()
 
 	SessionStartTime = FPlatformTime::Seconds();
 
-	// Install Editor Mode callbacks
+	// Install Editor Only Mode callbacks. Do not record these for Editor Commandlet runs 
+	if (GIsEditor==true && IsRunningCommandlet()==false)
+	{
+		// Start Editor and Editor Boot span. Note : this will only start when the plugin is loaded and as such will miss any activity that runs beforehand
+		EditorSpan = FStudioTelemetry::Get().StartSpan(EditorSpanName);
+		EditorBootSpan = FStudioTelemetry::Get().StartSpan(EditorBootSpanName, EditorSpan);
+		EditorMapName = TEXT("None");
+
+		TArray<FAnalyticsEventAttribute> Attributes;
+		Attributes.Emplace(FAnalyticsEventAttribute(TEXT("MapName"), EditorMapName));
+		EditorBootSpan->AddAttributes(Attributes);
 	
-	// Start Editor and Editor Boot span. Note : this will only start when the plugin is loaded and as such will miss any activity that runs beforehand
-	EditorSpan = FStudioTelemetry::Get().StartSpan(EditorSpanName);
-	EditorBootSpan = FStudioTelemetry::Get().StartSpan(EditorBootSpanName, EditorSpan);
-	EditorMapName = TEXT("None");
+		FEditorDelegates::OnEditorBoot.AddLambda([this](double)
+			{
+				TArray<FAnalyticsEventAttribute> Attributes;
+				Attributes.Emplace(TEXT("PluginCount"), TotalPluginCount);
+				EditorBootSpan->AddAttributes(Attributes);
 
-	TArray<FAnalyticsEventAttribute> Attributes;
-	Attributes.Emplace(FAnalyticsEventAttribute(TEXT("MapName"), EditorMapName));
-	EditorBootSpan->AddAttributes(Attributes);
+				FStudioTelemetry::Get().EndSpan(EditorBootSpan);
 
-	// Set up Slow Task callbacks
-	ensureMsgf(GWarn, TEXT("GWarn was not valid"));
+				// Callback is received when the editor has booted but has not been initialized
+				FEditorTelemetry::RecordEvent_Loading(TEXT("BootEditor"), EditorBootSpan->GetDuration(), EditorBootSpan->GetAttributes());
+				FEditorTelemetry::RecordEvent_CoreSystems(TEXT("BootEditor"), EditorBootSpan->GetAttributes());
 
-	if (GWarn != nullptr)
-	{
-		// Start the SlowTask span
-		GWarn->OnStartSlowTaskWithGuid().AddLambda([this](FGuid TaskGuid, const FText& TaskName)
+				EditorInitilizeSpan = FStudioTelemetry::Get().StartSpan(EditorInitilizeSpanName, EditorSpan);
+			});
+
+		FEditorDelegates::OnEditorInitialized.AddLambda([this](double TimeToInitializeEditor)
 			{
-				// Slow tasks can possibly be started from multiple threads, so we need to protect the registered span table
-				FScopeLock ScopeLock(&TaskSpanCriticalSection);
+				TimeToBootEditor = TimeToInitializeEditor;
 
-				TSharedPtr<IAnalyticsSpan>* SpanPtr = TaskSpans.Find(TaskGuid);
+				// Editor has initialized
+				TArray<FAnalyticsEventAttribute> Attributes;
+				Attributes.Emplace(TEXT("MapName"), EditorMapName);
+				Attributes.Emplace(TEXT("PluginCount"), TotalPluginCount);
 
-				// Only one task with this Guid is running asynchronously is supported at this time.
-				if (SpanPtr == nullptr)
-				{
-					TArray<FAnalyticsEventAttribute> Attributes;
-					Attributes.Emplace(TEXT("MapName"), EditorMapName);
-					Attributes.Emplace(TEXT("TaskName"), TaskName.ToString());
+				// Editor has finished initializing so start the Editor Interact span
+				FStudioTelemetry::Get().EndSpan(EditorInitilizeSpan);
+				EditorInteractSpan = FStudioTelemetry::Get().StartSpan(EditorInteractSpanName, EditorSpan);
 
-					// Create and start a new slow task span
-					TSharedPtr<IAnalyticsSpan> SlowTaskSpan = FStudioTelemetry::Get().StartSpan(TEXT("SlowTask"), Attributes);
+				FEditorTelemetry::RecordEvent_Loading(TEXT("TotalEditorStartup"), TimeToInitializeEditor, Attributes);
+				FEditorTelemetry::RecordEvent_CoreSystems(TEXT("TotalEditorStartup"), Attributes);
 
-					// Store this SlowTask span so we can find it when it finishes
-					TaskSpans.Add(TaskGuid, SlowTaskSpan);
-				}
+				// Install callbacks for Open Asset Dialogue
+				GEditor->GetEditorSubsystem<UAssetEditorSubsystem>()->OnAssetEditorRequestedOpen().AddLambda([this](UObject* Asset)
+					{
+						AssetOpenStartTime = FPlatformTime::Seconds();
+						FStudioTelemetry::Get().StartSpan(OpenAssetEditorSpan);
+					});
 
-				TRACE_BEGIN_REGION(*TaskName.ToString());
+				GEditor->GetEditorSubsystem<UAssetEditorSubsystem>()->OnAssetOpenedInEditor().AddLambda([this](UObject* Asset, IAssetEditorInstance*)
+					{
+						TArray<FAnalyticsEventAttribute> Attributes;
+						Attributes.Emplace(TEXT("MapName"), EditorMapName);
+
+						if (Asset != nullptr)
+						{
+							Attributes.Emplace(TEXT("AssetPath"), Asset->GetFullName());
+							Attributes.Emplace(TEXT("AssetClass"), Asset->GetClass()->GetName());
+
+							FEditorTelemetry::RecordEvent_Loading(TEXT("OpenAssetEditor"), FPlatformTime::Seconds() - AssetOpenStartTime, Attributes);
+						}
+
+						FStudioTelemetry::Get().EndSpan(OpenAssetEditorSpan, Attributes);
+					});
+
+				// Setup a timer for a Heartbeat callback
+				FTimerDelegate HeartbeatDelegate;
+				HeartbeatDelegate.BindRaw(this, &FEditorTelemetry::HeartbeatCallback);
+				GEditor->GetTimerManager()->SetTimer(TelemetryHeartbeatTimerHandle, HeartbeatDelegate, HeartbeatIntervalSeconds, true);
+
+				// Setup the timer for the Hitch Detector callback
+				FTimerDelegate HitchSamplerDelegate;
+				HitchSamplerDelegate.BindRaw(this, &FEditorTelemetry::HitchSamplerCallback);
+				GEditor->GetTimerManager()->SetTimer(TelemetryHitchSamplerTimerHandle, HitchSamplerDelegate, HitchSamplerIntervalSeconds, true);
 			});
 
-		// End the SlowTask span
-		GWarn->OnFinalizeSlowTaskWithGuid().AddLambda([this](FGuid TaskGuid, const FText& TaskName)
+		// Install PIE Mode callbacks
+		FEditorDelegates::StartPIE.AddLambda([this](bool)
 			{
-				TRACE_END_REGION(*TaskName.ToString());
+				// PIE mode has been started. The user has pressed the Start PIE button.
+				// Finish the Editor span
+				FStudioTelemetry::Get().EndSpan(EditorSpan);
+				EditorSessionCount++;
 
-				// Slow tasks can possibly be finalized from multiple threads, so we need to protect the registered span table
-				FScopeLock ScopeLock(&TaskSpanCriticalSection);
+				// Start PIE span
+				PIESpan = FStudioTelemetry::Get().StartSpan(PIESpanName);
 
-				// Find the task we stored off when we started this task
-				TSharedPtr<IAnalyticsSpan>* SpanPtr = TaskSpans.Find(TaskGuid);
+				// Append the PIE transition count to the PIE name
+				PIEStartupSpan = FStudioTelemetry::Get().StartSpan(PIESessionCount == 0 ? PIEStartupSpanName : FName(*FString::Printf(TEXT("%s%d"), *PIEStartupSpanName.ToString(), PIESessionCount)), PIESpan);
 
-				if (SpanPtr != nullptr)
+				TArray<FAnalyticsEventAttribute> Attributes;
+				Attributes.Emplace(TEXT("MapName"), EditorMapName);
+
+				PIESpan->AddAttributes(Attributes);
+				PIEStartupSpan->AddAttributes(Attributes);
+			});
+
+		FEditorDelegates::PreBeginPIE.AddLambda([this](bool)
+			{
+				PIEPreBeginSpan = FStudioTelemetry::Get().StartSpan(PIESessionCount == 0 ? PIEPreBeginSpanName : FName(*FString::Printf(TEXT("%s%d"), *PIEPreBeginSpanName.ToString(), PIESessionCount)), PIESpan);
+				PIEPreBeginSpan->AddAttributes(PIESpan->GetAttributes());
+			});
+
+		FEditorDelegates::BeginPIE.AddLambda([this](bool)
+			{
+				FStudioTelemetry::Get().EndSpan(PIEPreBeginSpan);
+			});
+
+		FWorldDelegates::OnPIEMapCreated.AddLambda([this](UGameInstance* GameInstance)
+			{
+				// A new PIE map was created
+				PIELoadMapSpan = FStudioTelemetry::Get().StartSpan(PIELoadMapSpanName, PIEStartupSpan);
+			});
+
+		FWorldDelegates::OnPIEMapReady.AddLambda([this](UGameInstance* GameInstance)
+			{
+				// PIE map is now loaded and ready to use
+				PIEMapName = FPaths::GetBaseFilename(GameInstance->PIEMapName);
+
+				TArray<FAnalyticsEventAttribute> Attributes;
+				Attributes.Emplace(TEXT("MapName"), EditorMapName);
+				Attributes.Emplace(TEXT("PIE_MapName"), PIEMapName);
+
+				PIELoadMapSpan->AddAttributes(Attributes);
+
+				FStudioTelemetry::Get().EndSpan(PIELoadMapSpan);
+
+				FEditorTelemetry::RecordEvent_Loading(TEXT("PIE.LoadMapTime"), PIELoadMapSpan->GetDuration(), PIELoadMapSpan->GetAttributes());
+				FEditorTelemetry::RecordEvent_CoreSystems(TEXT("PIE.LoadMapTime"), PIELoadMapSpan->GetAttributes());
+			});
+
+		FWorldDelegates::OnPIEReady.AddLambda([this](UGameInstance* GameInstance)
+			{
+				if (PIESpan.IsValid())
 				{
-					TSharedPtr<IAnalyticsSpan> SlowTaskSpan = *SpanPtr;
+					if (PIEStartupSpan.IsValid())
+					{
+						// PIE is now ready for user interaction
 
-					FStudioTelemetry::Get().EndSpan(SlowTaskSpan);
+						// Keep track of the PIE transition counts
+						TArray<FAnalyticsEventAttribute> Attributes;
+						Attributes.Emplace(TEXT("PIE_TransitionCount"), PIESessionCount);
 
-					// Remove the SlowTask span from the registry
-					TaskSpans.Remove(TaskGuid);
+						PIESpan->AddAttributes(Attributes);
+						PIEStartupSpan->AddAttributes(Attributes);
+
+						FStudioTelemetry::Get().EndSpan(PIEStartupSpan);
+
+						// Record the time from start PIE to PIE
+						FEditorTelemetry::RecordEvent_Loading(TEXT("PIE.TotalStartupTime"), PIEStartupSpan->GetDuration(), PIEStartupSpan->GetAttributes());
+						FEditorTelemetry::RecordEvent_CoreSystems(TEXT("PIE.TotalStartupTime"), PIEStartupSpan->GetAttributes());
+
+						if (PIESessionCount == 0)
+						{
+							const double TimeInEditor = EditorLoadMapSpan.IsValid() ? EditorLoadMapSpan->GetDuration() : 0.0;
+							const double TimeToStartPIE = PIEStartupSpan->GetDuration();
+							const double TimeToBootToPIE = TimeToBootEditor + TimeInEditor + TimeToStartPIE;
+
+							// Record the absolute time from editor boot to PIE
+							FEditorTelemetry::RecordEvent_Loading(TEXT("TimeToPIE"), TimeToBootToPIE, PIEStartupSpan->GetAttributes());
+							FEditorTelemetry::RecordEvent_CoreSystems(TEXT("TimeToPIE"), PIEStartupSpan->GetAttributes());
+						}
+					}
+
+					PIEInteractSpan = FStudioTelemetry::Get().StartSpan(PIESessionCount == 0 ? PIEInteractSpanName : FName(*FString::Printf(TEXT("%s%d"), *PIEInteractSpanName.ToString(), PIESessionCount)), PIESpan);
+					PIEInteractSpan->AddAttributes(PIESpan->GetAttributes());
 				}
 			});
-	}
 
-	FEditorDelegates::OnMapLoad.AddLambda([this](const FString& MapName, FCanLoadMap& OutCanLoadMap)
-		{
-			if (MapName.Len() > 0)
+		FEditorDelegates::EndPIE.AddLambda([this](bool)
 			{
-				// The Editor loads a new map
-				EditorLoadMapSpan = FStudioTelemetry::Get().StartSpan(EditorLoadMapSpanName, EditorSpan);
-			}
-		});
+				if (PIESpan.IsValid())
+				{
+					// PIE is ending so no longer interactive
+					FStudioTelemetry::Get().EndSpan(PIEInteractSpan);
+					PIEShutdownSpan = FStudioTelemetry::Get().StartSpan(PIESessionCount == 0 ? PIEShutdownSpanName : FName(*FString::Printf(TEXT("%s%d"), *PIEShutdownSpanName.ToString(), PIESessionCount)), PIESpan);
+					PIEShutdownSpan->AddAttributes(PIESpan->GetAttributes());
+				}
+			});
 
-	FEditorDelegates::OnMapOpened.AddLambda([this](const FString& MapName, bool Unused)
-		{
-			if (EditorLoadMapSpan.IsValid())
+		FEditorDelegates::ShutdownPIE.AddLambda([this](bool)
 			{
-				// The new editor map was actually opened
-				EditorMapName = FPaths::GetBaseFilename(MapName);
+				if (PIESpan.IsValid())
+				{
+					// PIE has shutdown, ie. the user has pressed the Stop PIE button, and we are going back to interactive Editor mode	
+					FStudioTelemetry::Get().EndSpan(PIESpan);
+
+					FEditorTelemetry::RecordEvent_Loading(TEXT("PIE.EndTime"), PIESpan->GetDuration(), PIESpan->GetAttributes());
+					FEditorTelemetry::RecordEvent_CoreSystems(TEXT("PIE.EndTime"), PIESpan->GetAttributes());
+				}
+
+				PIESessionCount++;
 
 				TArray<FAnalyticsEventAttribute> Attributes;
 				Attributes.Emplace(TEXT("MapName"), EditorMapName);
 
-				EditorSpan->AddAttributes(Attributes);
-				EditorLoadMapSpan->AddAttributes(Attributes);
+				// Restart the Editor span
+				EditorSpan = FStudioTelemetry::Get().StartSpan(EditorSpanName, Attributes);
+				EditorInteractSpan = FStudioTelemetry::Get().StartSpan(FName(*FString::Printf(TEXT("%s%d"), *EditorInteractSpanName.ToString(), EditorSessionCount)), EditorSpan);
+			});
 
-				FStudioTelemetry::Get().EndSpan(EditorLoadMapSpan);
-				
-				FEditorTelemetry::RecordEvent_Loading(TEXT("LoadMap"), EditorLoadMapSpan->GetDuration(), EditorLoadMapSpan->GetAttributes());
-				FEditorTelemetry::RecordEvent_CoreSystems(TEXT("LoadMap"), EditorLoadMapSpan->GetAttributes());
+		FEditorDelegates::OnMapLoad.AddLambda([this](const FString& MapName, FCanLoadMap& OutCanLoadMap)
+			{
+				if (MapName.Len() > 0)
+				{
+					// The Editor loads a new map
+					EditorLoadMapSpan = FStudioTelemetry::Get().StartSpan(EditorLoadMapSpanName, EditorSpan);
+				}
+			});
 
-				EditorLoadMapSpan.Reset();
-			}	
-		});
+		FEditorDelegates::OnMapOpened.AddLambda([this](const FString& MapName, bool Unused)
+			{
+				if (EditorLoadMapSpan.IsValid())
+				{
+					// The new editor map was actually opened
+					EditorMapName = FPaths::GetBaseFilename(MapName);
+
+					TArray<FAnalyticsEventAttribute> Attributes;
+					Attributes.Emplace(TEXT("MapName"), EditorMapName);
+
+					EditorSpan->AddAttributes(Attributes);
+					EditorLoadMapSpan->AddAttributes(Attributes);
+
+					FStudioTelemetry::Get().EndSpan(EditorLoadMapSpan);
+
+					FEditorTelemetry::RecordEvent_Loading(TEXT("LoadMap"), EditorLoadMapSpan->GetDuration(), EditorLoadMapSpan->GetAttributes());
+					FEditorTelemetry::RecordEvent_CoreSystems(TEXT("LoadMap"), EditorLoadMapSpan->GetAttributes());
+
+					EditorLoadMapSpan.Reset();
+				}
+			});
+	}
 
 	// Install any plugin load/unload callbacks
 	FModuleManager::Get().OnModulesChanged().AddLambda([this](FName ModuleName, EModuleChangeReason ChangeReason)
 		{
-			switch ( ChangeReason )
-			{ 
+			switch (ChangeReason)
+			{
 				default:
 				{
 					break;
 				}
-				
+
 				case EModuleChangeReason::ModuleLoaded:
 				{
 					TotalPluginCount++;
@@ -557,7 +696,7 @@ void FEditorTelemetry::StartSession()
 								{
 									// Start the Asset Registry Scan span
 									AssetRegistryScanSpan = FStudioTelemetry::Get().StartSpan(AssetRegistryScanSpanName, EditorSpan);
-								}	
+								}
 
 								AssetRegistryScanCount++;
 							});
@@ -572,239 +711,91 @@ void FEditorTelemetry::StartSession()
 									TArray<FAnalyticsEventAttribute> Attributes;
 									Attributes.Emplace(TEXT("MapName"), EditorMapName);
 									FStudioTelemetry::Get().EndSpan(AssetRegistryScanSpan, Attributes);
-								}	
+								}
 							});
 					}
 
 					break;
 				}
-		
+
 				case EModuleChangeReason::ModuleUnloaded:
 				{
 					TotalPluginCount--;
 					break;
-				}	
+				}
 			}
 		});
+	
+	// Set up Slow Task callbacks
+	ensureMsgf(GWarn, TEXT("GWarn was not valid"));
 
-	FEditorDelegates::OnEditorBoot.AddLambda([this](double)
-		{	
-			TArray<FAnalyticsEventAttribute> Attributes;
-			Attributes.Emplace(TEXT("PluginCount"), TotalPluginCount);
-			EditorBootSpan->AddAttributes(Attributes);
-
-			FStudioTelemetry::Get().EndSpan(EditorBootSpan);
-
-			// Callback is received when the editor has booted but has not been initialized
-			FEditorTelemetry::RecordEvent_Loading(TEXT("BootEditor"), EditorBootSpan->GetDuration(), EditorBootSpan->GetAttributes());
-			FEditorTelemetry::RecordEvent_CoreSystems(TEXT("BootEditor"), EditorBootSpan->GetAttributes());
-
-			EditorInitilizeSpan = FStudioTelemetry::Get().StartSpan(EditorInitilizeSpanName, EditorSpan);
-		});
-
-	FEditorDelegates::OnEditorInitialized.AddLambda([this](double TimeToInitializeEditor)
-		{
-			TimeToBootEditor = TimeToInitializeEditor;
-
-			// Editor has initialized
-			TArray<FAnalyticsEventAttribute> Attributes;
-			Attributes.Emplace(TEXT("MapName"), EditorMapName);
-			Attributes.Emplace(TEXT("PluginCount"), TotalPluginCount);
-
-			// Editor has finished initializing so start the Editor Interact span
-			FStudioTelemetry::Get().EndSpan(EditorInitilizeSpan);
-			EditorInteractSpan = FStudioTelemetry::Get().StartSpan(EditorInteractSpanName, EditorSpan);
-						
-			FEditorTelemetry::RecordEvent_Loading(TEXT("TotalEditorStartup"), TimeToInitializeEditor, Attributes);
-			FEditorTelemetry::RecordEvent_CoreSystems(TEXT("TotalEditorStartup"), Attributes);
-
-			ensureMsgf(GEditor, TEXT("GEditor was not valid"));
-
-			if (GEditor != nullptr)
+	if (GWarn != nullptr)
+	{
+		// Start the SlowTask span
+		GWarn->OnStartSlowTaskWithGuid().AddLambda([this](FGuid TaskGuid, const FText& TaskName)
 			{
-				// Install callbacks for Open Asset Dialogue
-				GEditor->GetEditorSubsystem<UAssetEditorSubsystem>()->OnAssetEditorRequestedOpen().AddLambda([this](UObject* Asset)
-					{
-						AssetOpenStartTime = FPlatformTime::Seconds();
-						FStudioTelemetry::Get().StartSpan(OpenAssetEditorSpan);
-					});
-
-				GEditor->GetEditorSubsystem<UAssetEditorSubsystem>()->OnAssetOpenedInEditor().AddLambda([this](UObject* Asset, IAssetEditorInstance*)
-					{
-						TArray<FAnalyticsEventAttribute> Attributes;
-						Attributes.Emplace(TEXT("MapName"), EditorMapName);
-						
-						if (Asset != nullptr)
-						{
-							Attributes.Emplace(TEXT("AssetPath"), Asset->GetFullName());
-							Attributes.Emplace(TEXT("AssetClass"), Asset->GetClass()->GetName());
-
-							FEditorTelemetry::RecordEvent_Loading(TEXT("OpenAssetEditor"), FPlatformTime::Seconds() - AssetOpenStartTime, Attributes);
-						}
-
-						FStudioTelemetry::Get().EndSpan(OpenAssetEditorSpan, Attributes);			
-				});
+				// Slow tasks can possibly be started from multiple threads, so we need to protect the registered span table
+				FScopeLock ScopeLock(&TaskSpanCriticalSection);
 
-				// Setup a timer for a Heartbeat callback
-				FTimerDelegate HeartbeatDelegate;
-				HeartbeatDelegate.BindRaw(this, &FEditorTelemetry::HeartbeatCallback);
-				GEditor->GetTimerManager()->SetTimer(TelemetryHeartbeatTimerHandle, HeartbeatDelegate, HeartbeatIntervalSeconds, true);
+				TSharedPtr<IAnalyticsSpan>* SpanPtr = TaskSpans.Find(TaskGuid);
 
-				// Setup the timer for the Hitch Detector callback
-				if (IsRunningCommandlet() == false)
+				// Only one task with this Guid is running asynchronously is supported at this time.
+				if (SpanPtr == nullptr)
 				{
-					FTimerDelegate HitchSamplerDelegate;
-					HitchSamplerDelegate.BindRaw(this, &FEditorTelemetry::HitchSamplerCallback);
-					GEditor->GetTimerManager()->SetTimer(TelemetryHitchSamplerTimerHandle, HitchSamplerDelegate, HitchSamplerIntervalSeconds, true);
-				}
-			}
-
-			ensureMsgf(GUnrealEd, TEXT("GUnrealEd was not valid"));
-			
-			if (GUnrealEd!=nullptr && GUnrealEd->CookServer != nullptr )
-			{
-				UE::Cook::FDelegates::CookFinished.AddLambda([this](UE::Cook::ICookInfo& CookInfo)
-					{
-						if (CookInfo.GetCookType() != UE::Cook::ECookType::ByTheBook)
-						{
-							return;
-						}
-						TArray<FAnalyticsEventAttribute> Attributes;
-						Attributes.Emplace(TEXT("MapName"), EditorMapName);
-						
-						FEditorTelemetry::RecordEvent_Cooking(Attributes);
-						FEditorTelemetry::RecordEvent_CoreSystems(TEXT("Cooking"), Attributes);
-					});
-			}
-		});
-
-	// Install PIE Mode callbacks
-	FEditorDelegates::StartPIE.AddLambda([this](bool)
-		{
-			// PIE mode has been started. The user has pressed the Start PIE button.
-			// Finish the Editor span
-			FStudioTelemetry::Get().EndSpan(EditorSpan);
-			EditorSessionCount++;
-
-			// Start PIE span
-			PIESpan = FStudioTelemetry::Get().StartSpan(PIESpanName);
+					TArray<FAnalyticsEventAttribute> Attributes;
+					Attributes.Emplace(TEXT("MapName"), EditorMapName);
+					Attributes.Emplace(TEXT("TaskName"), TaskName.ToString());
 
-			// Append the PIE transition count to the PIE name
-			PIEStartupSpan = FStudioTelemetry::Get().StartSpan(PIESessionCount == 0? PIEStartupSpanName : FName(*FString::Printf(TEXT("%s%d"), *PIEStartupSpanName.ToString(), PIESessionCount)), PIESpan);
-		
-			TArray<FAnalyticsEventAttribute> Attributes;
-			Attributes.Emplace(TEXT("MapName"), EditorMapName);
+					// Create and start a new slow task span
+					TSharedPtr<IAnalyticsSpan> SlowTaskSpan = FStudioTelemetry::Get().StartSpan(TEXT("SlowTask"), Attributes);
 
-			PIESpan->AddAttributes(Attributes);
-			PIEStartupSpan->AddAttributes(Attributes);
-		});
+					// Store this SlowTask span so we can find it when it finishes
+					TaskSpans.Add(TaskGuid, SlowTaskSpan);
+				}
 
-	FEditorDelegates::PreBeginPIE.AddLambda([this](bool)
-		{	
-			PIEPreBeginSpan = FStudioTelemetry::Get().StartSpan(PIESessionCount == 0 ? PIEPreBeginSpanName : FName(*FString::Printf(TEXT("%s%d"), *PIEPreBeginSpanName.ToString(), PIESessionCount)), PIESpan);
-			PIEPreBeginSpan->AddAttributes(PIESpan->GetAttributes());
-		});
+				TRACE_BEGIN_REGION(*TaskName.ToString());
+			});
 
-	FEditorDelegates::BeginPIE.AddLambda([this](bool)
-		{				
-			FStudioTelemetry::Get().EndSpan(PIEPreBeginSpan);
-		});
+		// End the SlowTask span
+		GWarn->OnFinalizeSlowTaskWithGuid().AddLambda([this](FGuid TaskGuid, const FText& TaskName)
+			{
+				TRACE_END_REGION(*TaskName.ToString());
 
-	FWorldDelegates::OnPIEMapCreated.AddLambda([this](UGameInstance* GameInstance)
-		{
-			// A new PIE map was created
-			PIELoadMapSpan = FStudioTelemetry::Get().StartSpan(PIELoadMapSpanName, PIEStartupSpan);
-		});
+				// Slow tasks can possibly be finalized from multiple threads, so we need to protect the registered span table
+				FScopeLock ScopeLock(&TaskSpanCriticalSection);
 
-	FWorldDelegates::OnPIEMapReady.AddLambda([this](UGameInstance* GameInstance)
-		{
-			// PIE map is now loaded and ready to use
-			PIEMapName = FPaths::GetBaseFilename(GameInstance->PIEMapName);
+				// Find the task we stored off when we started this task
+				TSharedPtr<IAnalyticsSpan>* SpanPtr = TaskSpans.Find(TaskGuid);
 
-			TArray<FAnalyticsEventAttribute> Attributes;
-			Attributes.Emplace(TEXT("MapName"), EditorMapName);
-			Attributes.Emplace(TEXT("PIE_MapName"), PIEMapName);
+				if (SpanPtr != nullptr)
+				{
+					TSharedPtr<IAnalyticsSpan> SlowTaskSpan = *SpanPtr;
 
-			PIELoadMapSpan->AddAttributes(Attributes);
-			
-			FStudioTelemetry::Get().EndSpan(PIELoadMapSpan);
+					FStudioTelemetry::Get().EndSpan(SlowTaskSpan);
 
-			FEditorTelemetry::RecordEvent_Loading(TEXT("PIE.LoadMapTime"), PIELoadMapSpan->GetDuration(), PIELoadMapSpan->GetAttributes());
-			FEditorTelemetry::RecordEvent_CoreSystems(TEXT("PIE.LoadMapTime"), PIELoadMapSpan->GetAttributes());
-		});
+					// Remove the SlowTask span from the registry
+					TaskSpans.Remove(TaskGuid);
+				}
+			});
+	}
 
-	FWorldDelegates::OnPIEReady.AddLambda([this](UGameInstance* GameInstance)
-		{
-			if (PIESpan.IsValid())
+	// Install Cooking Callbacks
+	if (GUnrealEd != nullptr && GUnrealEd->CookServer != nullptr)
+	{
+		UE::Cook::FDelegates::CookFinished.AddLambda([this](UE::Cook::ICookInfo& CookInfo)
 			{
-				if ( PIEStartupSpan.IsValid() )
+				if (CookInfo.GetCookType() != UE::Cook::ECookType::ByTheBook)
 				{
-					// PIE is now ready for user interaction
-					
-					// Keep track of the PIE transition counts
-					TArray<FAnalyticsEventAttribute> Attributes;
-					Attributes.Emplace(TEXT("PIE_TransitionCount"), PIESessionCount );
-					
-					PIESpan->AddAttributes(Attributes);
-					PIEStartupSpan->AddAttributes(Attributes);
-						
-					FStudioTelemetry::Get().EndSpan(PIEStartupSpan);
-
-					// Record the time from start PIE to PIE
-					FEditorTelemetry::RecordEvent_Loading(TEXT("PIE.TotalStartupTime"), PIEStartupSpan->GetDuration(), PIEStartupSpan->GetAttributes());
-					FEditorTelemetry::RecordEvent_CoreSystems(TEXT("PIE.TotalStartupTime"), PIEStartupSpan->GetAttributes());
-
-					if (PIESessionCount == 0)
-					{
-						const double TimeInEditor = EditorLoadMapSpan.IsValid() ? EditorLoadMapSpan->GetDuration() : 0.0;
-						const double TimeToStartPIE = PIEStartupSpan->GetDuration();
-						const double TimeToBootToPIE = TimeToBootEditor + TimeInEditor + TimeToStartPIE;
-
-						// Record the absolute time from editor boot to PIE
-						FEditorTelemetry::RecordEvent_Loading(TEXT("TimeToPIE"), TimeToBootToPIE, PIEStartupSpan->GetAttributes());
-						FEditorTelemetry::RecordEvent_CoreSystems(TEXT("TimeToPIE"), PIEStartupSpan->GetAttributes());
-					}
+					return;
 				}
+				TArray<FAnalyticsEventAttribute> Attributes;
+				Attributes.Emplace(TEXT("MapName"), EditorMapName);
 
-				PIEInteractSpan = FStudioTelemetry::Get().StartSpan( PIESessionCount==0? PIEInteractSpanName : FName(*FString::Printf(TEXT("%s%d"), *PIEInteractSpanName.ToString(), PIESessionCount)), PIESpan);
-				PIEInteractSpan->AddAttributes(PIESpan->GetAttributes());
-			}
-		});
-
-	FEditorDelegates::EndPIE.AddLambda([this](bool)
-		{
-			if (PIESpan.IsValid())
-			{
-				// PIE is ending so no longer interactive
-				FStudioTelemetry::Get().EndSpan(PIEInteractSpan);
-				PIEShutdownSpan = FStudioTelemetry::Get().StartSpan(PIESessionCount == 0 ? PIEShutdownSpanName : FName(*FString::Printf(TEXT("%s%d"), *PIEShutdownSpanName.ToString(), PIESessionCount)), PIESpan);
-				PIEShutdownSpan->AddAttributes(PIESpan->GetAttributes());
-			}
-		});
-
-	FEditorDelegates::ShutdownPIE.AddLambda([this](bool)
-		{
-			if (PIESpan.IsValid())
-			{
-				// PIE has shutdown, ie. the user has pressed the Stop PIE button, and we are going back to interactive Editor mode	
-				FStudioTelemetry::Get().EndSpan(PIESpan);
-
-				FEditorTelemetry::RecordEvent_Loading(TEXT("PIE.EndTime"), PIESpan->GetDuration(), PIESpan->GetAttributes());
-				FEditorTelemetry::RecordEvent_CoreSystems(TEXT("PIE.EndTime"), PIESpan->GetAttributes());
-			}
-
-			PIESessionCount++;
-
-			TArray<FAnalyticsEventAttribute> Attributes;
-			Attributes.Emplace(TEXT("MapName"), EditorMapName);
-
-			// Restart the Editor span
-			EditorSpan = FStudioTelemetry::Get().StartSpan(EditorSpanName, Attributes);
-			EditorInteractSpan = FStudioTelemetry::Get().StartSpan(FName(*FString::Printf(TEXT("%s%d"), *EditorInteractSpanName.ToString(), EditorSessionCount)), EditorSpan);
-		});
+				FEditorTelemetry::RecordEvent_Cooking(Attributes);
+				FEditorTelemetry::RecordEvent_CoreSystems(TEXT("Cooking"), Attributes);
+			});
+	}
 
-	
-	// Install Cooking Callbacks
 	UE::Cook::FDelegates::CookStarted.AddLambda([this](UE::Cook::ICookInfo& CookInfo)
 	{
 		if (CookInfo.GetCookType() != UE::Cook::ECookType::ByTheBook)
