Added GetElapsedTime() to the IAnalyticsProvder and FAnalyticsProvider implementation
Added saftey check for divde by zero when calculating the HitchesPerSecond value
Made sure that the MapName and PIE_MapNames are correctly added to the attributes

diff --git a/Engine/Plugins/Experimental/EditorTelemetry/Source/Private/EditorTelemetry.cpp b/Engine/Plugins/Experimental/EditorTelemetry/Source/Private/EditorTelemetry.cpp
--- a/Engine/Plugins/Experimental/EditorTelemetry/Source/Private/EditorTelemetry.cpp
+++ b/Engine/Plugins/Experimental/EditorTelemetry/Source/Private/EditorTelemetry.cpp
@@ -396,9 +396,12 @@ void FEditorTelemetry::HeartbeatCallback()
 		{
 			// No longer hitching and we have started a hitch span
 			TArray<FAnalyticsEventAttribute> Attributes;
+			const double ElapsedTime = HitchingSpan->GetElapsedTime();
 			Attributes.Emplace(FAnalyticsEventAttribute(TEXT("Hitch_Count"), HitchCount));
-			Attributes.Emplace(FAnalyticsEventAttribute(TEXT("Hitch_HitchesPerSecond"), (float)HitchCount / HitchingSpan->GetDuration()));
+			Attributes.Emplace(FAnalyticsEventAttribute(TEXT("Hitch_HitchesPerSecond"), ElapsedTime>0? (float)HitchCount / ElapsedTime : 0.0));
 			Attributes.Emplace(FAnalyticsEventAttribute(TEXT("Hitch_AverageFPS"), HitchAvergageFPS));
+			Attributes.Emplace(TEXT("MapName"), EditorMapName);
+			Attributes.Emplace(TEXT("PIE_MapName"), PIEMapName);
 
 			// End the hitch Span
 			FStudioTelemetry::Get().EndSpan(HitchingSpan, Attributes);
@@ -406,9 +409,6 @@ void FEditorTelemetry::HeartbeatCallback()
 			// Record the hitch event
 			FStudioTelemetry::Get().RecordEvent(TEXT("Core.Hitch"), Attributes);
 
-			Attributes.Emplace(TEXT("MapName"), EditorMapName);
-			Attributes.Emplace(TEXT("PIE_MapName"), PIEMapName);
-
 			// Record core systems events for the hitch
 			RecordEvent_CoreSystems(TEXT("Hitch"));
 
diff --git a/Engine/Source/Runtime/Analytics/Analytics/Private/AnalyticsTracer.cpp b/Engine/Source/Runtime/Analytics/Analytics/Private/AnalyticsTracer.cpp
--- a/Engine/Source/Runtime/Analytics/Analytics/Private/AnalyticsTracer.cpp
+++ b/Engine/Source/Runtime/Analytics/Analytics/Private/AnalyticsTracer.cpp
@@ -42,6 +42,11 @@ void FAnalyticsSpan::SetStackDepth(uint32 Depth)
 	StackDepth = Depth;
 }
 
+double FAnalyticsSpan::GetElapsedTime()
+{
+	return (FDateTime::UtcNow() - StartTime).GetTotalSeconds();
+}
+
 double FAnalyticsSpan::GetDuration() const 
 {
 	return Duration;
diff --git a/Engine/Source/Runtime/Analytics/Analytics/Public/AnalyticsTracer.h b/Engine/Source/Runtime/Analytics/Analytics/Public/AnalyticsTracer.h
--- a/Engine/Source/Runtime/Analytics/Analytics/Public/AnalyticsTracer.h
+++ b/Engine/Source/Runtime/Analytics/Analytics/Public/AnalyticsTracer.h
@@ -26,6 +26,7 @@ public:
 	virtual const TArray<FAnalyticsEventAttribute>& GetAttributes() const override;
 	virtual void SetStackDepth(uint32 Depth) override;
 	virtual uint32 GetStackDepth() const override;
+	virtual double GetElapsedTime() override;
 	virtual double GetDuration() const override;
 	virtual bool GetIsActive() const override;
 	virtual FGuid GetId() const override;
diff --git a/Engine/Source/Runtime/Analytics/Analytics/Public/Interfaces/IAnalyticsTracer.h b/Engine/Source/Runtime/Analytics/Analytics/Public/Interfaces/IAnalyticsTracer.h
--- a/Engine/Source/Runtime/Analytics/Analytics/Public/Interfaces/IAnalyticsTracer.h
+++ b/Engine/Source/Runtime/Analytics/Analytics/Public/Interfaces/IAnalyticsTracer.h
@@ -54,6 +54,9 @@ public:
 	/** Get the span duration in seconds */
 	virtual double GetDuration() const = 0;
 
+	/** Get the elapsed time for the span in seconds */
+	virtual double GetElapsedTime() = 0;
+
 	// Is the span active?
 	virtual bool GetIsActive() const = 0;
 
