Changed  logic so that EditorTelemetry sessions are started either once Studio Telemetry session has started OR immediately if the Studio Telemetry session has already started
Added some checks to ensure that EditorTelemetry session is only ended once

diff --git a/Engine/Plugins/Experimental/EditorTelemetry/Source/Private/EditorTelemetry.cpp b/Engine/Plugins/Experimental/EditorTelemetry/Source/Private/EditorTelemetry.cpp
--- a/Engine/Plugins/Experimental/EditorTelemetry/Source/Private/EditorTelemetry.cpp
+++ b/Engine/Plugins/Experimental/EditorTelemetry/Source/Private/EditorTelemetry.cpp
@@ -1019,5 +1019,9 @@ void FEditorTelemetry::StartSession()
 
 void FEditorTelemetry::EndSession()
 {
-	FStudioTelemetry::Get().EndSpan(EditorSpan);
+	if (EditorSpan.IsValid())
+	{
+		FStudioTelemetry::Get().EndSpan(EditorSpan);
+		EditorSpan.Reset();
+	}
 }
diff --git a/Engine/Plugins/Experimental/EditorTelemetry/Source/Private/EditorTelemetryModule.cpp b/Engine/Plugins/Experimental/EditorTelemetry/Source/Private/EditorTelemetryModule.cpp
--- a/Engine/Plugins/Experimental/EditorTelemetry/Source/Private/EditorTelemetryModule.cpp
+++ b/Engine/Plugins/Experimental/EditorTelemetry/Source/Private/EditorTelemetryModule.cpp
@@ -9,12 +9,17 @@ IMPLEMENT_MODULE(FEditorTelemetryModule, EditorTelemetry);
 
 void FEditorTelemetryModule::StartupModule()
 {
+	if (FStudioTelemetry::Get().IsAvailable())
+	{
+		FEditorTelemetry::Get().StartSession();
+	}
+	
 	FStudioTelemetry::Get().GetOnStartSession().AddLambda([this]()
 		{
 			FEditorTelemetry::Get().StartSession();
 		}
 	);
-	
+
 	FStudioTelemetry::Get().GetOnEndSession().AddLambda([this]()
 		{
 			FEditorTelemetry::Get().EndSession();
@@ -24,5 +29,5 @@ void FEditorTelemetryModule::StartupModule()
 
 void FEditorTelemetryModule::ShutdownModule()
 {
-	FEditorTelemetry::Get().EndSession();
+	
 }
