Minor tweaks and changes to Editor Telemetry events

diff --git a/Engine/Plugins/Experimental/EditorTelemetry/Source/Private/EditorTelemetry.cpp b/Engine/Plugins/Experimental/EditorTelemetry/Source/Private/EditorTelemetry.cpp
--- a/Engine/Plugins/Experimental/EditorTelemetry/Source/Private/EditorTelemetry.cpp
+++ b/Engine/Plugins/Experimental/EditorTelemetry/Source/Private/EditorTelemetry.cpp
@@ -515,13 +515,17 @@ void FEditorTelemetry::StartSession()
 				TArray<FAnalyticsEventAttribute> Attributes;
 				Attributes.Emplace(TEXT("MapName"), EditorMapName);
 				Attributes.Emplace(TEXT("PluginCount"), TotalPluginCount);
-
+				EditorInitilizeSpan->AddAttributes(Attributes);
+				
 				// Editor has finished initializing so start the Editor Interact span
 				FStudioTelemetry::Get().EndSpan(EditorInitilizeSpan);
+
+				FEditorTelemetry::RecordEvent_Loading(TEXT("Editor.Initialize"), EditorInitilizeSpan->GetDuration(), EditorInitilizeSpan->GetAttributes());
+				FEditorTelemetry::RecordEvent_CoreSystems(TEXT("Editor.Initialize"), EditorInitilizeSpan->GetAttributes());
+				FEditorTelemetry::RecordEvent_Loading(TEXT("TimeToEdit"), TimeToInitializeEditor, EditorInitilizeSpan->GetAttributes());
+				
 				EditorInteractSpan = FStudioTelemetry::Get().StartSpan(EditorInteractSpanName, EditorSpan);
 
-				FEditorTelemetry::RecordEvent_Loading(TEXT("Editor.Initialize"), TimeToInitializeEditor, Attributes);
-				FEditorTelemetry::RecordEvent_CoreSystems(TEXT("Editor.Initialize"), Attributes);
 				
 				// Install callbacks for Open Asset Dialogue
 				GEditor->GetEditorSubsystem<UAssetEditorSubsystem>()->OnAssetEditorRequestedOpen().AddLambda([this](UObject* Asset)
@@ -627,6 +631,10 @@ void FEditorTelemetry::StartSession()
 
 						FStudioTelemetry::Get().EndSpan(PIEStartupSpan);
 
+						// Record the PIE startup
+						FEditorTelemetry::RecordEvent_Loading(TEXT("PIE.Startup"), PIEStartupSpan->GetDuration(), PIEStartupSpan->GetAttributes());
+						FEditorTelemetry::RecordEvent_CoreSystems(TEXT("PIE.Startup"), PIEStartupSpan->GetAttributes());
+
 						// Record the time from start PIE to PIE
 						if (PIESessionCount == 0)
 						{
@@ -636,10 +644,6 @@ void FEditorTelemetry::StartSession()
 
 							// Record the absolute time from editor boot to PIE
 							FEditorTelemetry::RecordEvent_Loading(TEXT("TimeToPIE"), TimeToBootToPIE, PIEStartupSpan->GetAttributes());	
-
-							// Record the first PIE transition
-							FEditorTelemetry::RecordEvent_Loading(TEXT("PIE.Startup"), PIEStartupSpan->GetDuration(), PIEStartupSpan->GetAttributes());
-							FEditorTelemetry::RecordEvent_CoreSystems(TEXT("PIE.Startup"), PIEStartupSpan->GetAttributes());
 						}
 					}
 
@@ -664,14 +668,11 @@ void FEditorTelemetry::StartSession()
 				if (PIESpan.IsValid())
 				{
 					// PIE has shutdown, ie. the user has pressed the Stop PIE button, and we are going back to interactive Editor mode	
+					FStudioTelemetry::Get().EndSpan(PIEShutdownSpan);
+					FEditorTelemetry::RecordEvent_Loading(TEXT("PIE.Shutdown"), PIEShutdownSpan->GetDuration(), PIEShutdownSpan->GetAttributes());
+					FEditorTelemetry::RecordEvent_CoreSystems(TEXT("PIE.Shutdown"), PIEShutdownSpan->GetAttributes());
+					
 					FStudioTelemetry::Get().EndSpan(PIESpan);
-
-					if (PIESessionCount == 0)
-					{
-						// Record the first PIE shutdown
-						FEditorTelemetry::RecordEvent_Loading(TEXT("PIE.Shutdown"), PIESpan->GetDuration(), PIESpan->GetAttributes());
-						FEditorTelemetry::RecordEvent_CoreSystems(TEXT("PIE.Shutdown"), PIESpan->GetAttributes());
-					}
 				}
 
 				PIESessionCount++;
