Adding FEditorTelemetry output for cook workers
- Used for per-worker DDC-Verify statistics
- May be used for other cook-worker telemetry later

diff --git a/Engine/Plugins/Experimental/EditorTelemetry/Source/Private/EditorTelemetry.cpp b/Engine/Plugins/Experimental/EditorTelemetry/Source/Private/EditorTelemetry.cpp
--- a/Engine/Plugins/Experimental/EditorTelemetry/Source/Private/EditorTelemetry.cpp
+++ b/Engine/Plugins/Experimental/EditorTelemetry/Source/Private/EditorTelemetry.cpp
@@ -286,7 +286,7 @@ void FEditorTelemetry::RecordEvent_DDCVerification(const FString& Context, TArra
 
 	for (const FDerivedDataCacheVerificationStat& Stat : DDCVerifyStats.Entries)
 	{
-		if (Stat.RecordsTested > 0)
+		if (Stat.RecordsTested > 0 || Stat.ValuesTested > 0)
 		{
 			TArray<FAnalyticsEventAttribute> EventAttributes = Attributes;
 			EventAttributes.Emplace(TEXT("Bucket"), Stat.BucketName);
@@ -897,21 +897,26 @@ void FEditorTelemetry::StartSession()
 		{
 			return;
 		}
-		// End the cooking span
 	
-		// Suppress sending telemetry from CookWorkers for now.
 		uint32 MultiprocessId = 0;
 		FParse::Value(FCommandLine::Get(), TEXT("-MultiprocessId="), MultiprocessId);
-		if (MultiprocessId != 0)
+		const bool bIsCookWorker = MultiprocessId != 0;
+
+		if (bIsCookWorker)
 		{
-			return;
+			// Cook-worker specific telemetry goes here
+			FEditorTelemetry::RecordEvent_DDCVerification(TEXT("Cooking"), CookingSpan->GetAttributes());
 		}
+		else
+		{
+			// Cook director/single-process cook telemetry goes here
+			FEditorTelemetry::RecordEvent_Cooking(CookingSpan->GetAttributes());
+			FEditorTelemetry::RecordEvent_CoreSystems(TEXT("Cooking"), CookingSpan->GetAttributes());
+			FEditorTelemetry::RecordEvent_DDCVerification(TEXT("Cooking"), CookingSpan->GetAttributes());
 
-		FEditorTelemetry::RecordEvent_Cooking(CookingSpan->GetAttributes());
-		FEditorTelemetry::RecordEvent_CoreSystems(TEXT("Cooking"), CookingSpan->GetAttributes());
-		FEditorTelemetry::RecordEvent_DDCVerification(TEXT("Cooking"), CookingSpan->GetAttributes());
-
-		FStudioTelemetry::Get().EndSpan(CookingSpan);
+			// End the cooking span
+			FStudioTelemetry::Get().EndSpan(CookingSpan);
+		}
 	});
 
 	// Install Content Browser callbacks
