DataValidation Telemetry

diff --git a/Engine/Plugins/Editor/DataValidation/Source/DataValidation/DataValidation.Build.cs b/Engine/Plugins/Editor/DataValidation/Source/DataValidation/DataValidation.Build.cs
--- a/Engine/Plugins/Editor/DataValidation/Source/DataValidation/DataValidation.Build.cs
+++ b/Engine/Plugins/Editor/DataValidation/Source/DataValidation/DataValidation.Build.cs
@@ -47,7 +47,8 @@ public class DataValidation : ModuleRules
 				"SourceControl",
 				"UncontrolledChangelists",
 				"RHI",
-				"EditorToolEvents"
+				"EditorToolEvents",
+				"TelemetryUtils",
 			}
 		);
 	}
diff --git a/Engine/Plugins/Editor/DataValidation/Source/DataValidation/Private/DataValidationTelemetry.cpp b/Engine/Plugins/Editor/DataValidation/Source/DataValidation/Private/DataValidationTelemetry.cpp
new file mode 100644
--- /dev/null
+++ b/Engine/Plugins/Editor/DataValidation/Source/DataValidation/Private/DataValidationTelemetry.cpp
@@ -0,0 +1,48 @@
+// Copyright Epic Games, Inc. All Rights Reserved.
+
+#include "DataValidationTelemetry.h"
+
+#include "CoreGlobals.h"
+#include "EditorValidatorSubsystem.h"
+#include "UObject/TopLevelAssetPath.h"
+
+namespace UE::Telemetry::DataValidation
+{
+
+FValidateAssets::FValidateAssets(const FValidateAssetsSettings& InSettings)
+{
+	Mode = IsRunningCommandlet() ? EValidationMode::Commandlet : EValidationMode::Editor;
+
+	bSkipExcludedDirectories = InSettings.bSkipExcludedDirectories;
+	ValidationUsecase = InSettings.ValidationUsecase;
+	bLoadAssetsForValidation = InSettings.bLoadAssetsForValidation;
+	bUnloadAssetsLoadedForValidation = InSettings.bUnloadAssetsLoadedForValidation;
+	bLoadExternalObjectsForValidation = InSettings.bLoadExternalObjectsForValidation;
+	bCaptureAssetLoadLogs = InSettings.bCaptureAssetLoadLogs;
+	bCaptureLogsDuringValidation = InSettings.bCaptureLogsDuringValidation;
+	bCaptureWarningsDuringValidationAsErrors = InSettings.bCaptureWarningsDuringValidationAsErrors;
+	MaxAssetsToValidate = InSettings.MaxAssetsToValidate;
+	bValidateReferencersOfDeletedAssets = InSettings.bValidateReferencersOfDeletedAssets;
+}
+
+void FValidateAssets::SetResults(EDataValidationResult InOverallResult, const FValidateAssetsResults& InResults)
+{
+	Result = InOverallResult;
+	NumRequested = InResults.NumRequested;
+	NumExternalObjects = InResults.NumExternalObjects;
+	NumChecked = InResults.NumChecked;
+	NumValid = InResults.NumValid;
+	NumInvalid = InResults.NumInvalid;
+	NumSkipped = InResults.NumSkipped;
+	NumWarnings = InResults.NumWarnings;
+	NumUnableToValidate = InResults.NumUnableToValidate;
+	bAssetLimitReached = InResults.bAssetLimitReached;
+	for (const TPair<FTopLevelAssetPath, ::FValidatorStatistics>& Pair : InResults.ValidatorStatistics)
+	{
+		FValidatorStatistics& Stats = ValidatorStatistics.FindOrAdd(Pair.Key);
+		Stats.AssetsValidated += Pair.Value.AssetsValidated;
+		Stats.AssetsAddedForValidation += Pair.Value.AssetsAddedForValidation;
+	}
+}
+
+}
\ No newline at end of file
diff --git a/Engine/Plugins/Editor/DataValidation/Source/DataValidation/Private/EditorValidatorSubsystem.cpp b/Engine/Plugins/Editor/DataValidation/Source/DataValidation/Private/EditorValidatorSubsystem.cpp
--- a/Engine/Plugins/Editor/DataValidation/Source/DataValidation/Private/EditorValidatorSubsystem.cpp
+++ b/Engine/Plugins/Editor/DataValidation/Source/DataValidation/Private/EditorValidatorSubsystem.cpp
@@ -9,7 +9,9 @@
 #include "AssetRegistry/AssetDataToken.h"
 #include "AssetRegistry/AssetRegistryModule.h"
 #include "Blueprint/BlueprintSupport.h"
+#include "CoreGlobals.h"
 #include "DataValidationChangelist.h"
+#include "DataValidationTelemetry.h"
 #include "EditorValidatorHelpers.h"
 #include "DirectoryWatcherModule.h"
 #include "Editor.h"
@@ -31,6 +33,7 @@
 #include "Misc/ScopedSlowTask.h"
 #include "Misc/ScopeExit.h"
 #include "PackageTools.h"
+#include "TelemetryRouter.h"
 #include "UObject/ICookInfo.h"
 #include "UObject/TopLevelAssetPath.h"
 
@@ -463,6 +466,17 @@ int32 UEditorValidatorSubsystem::ValidateAssetsWithSettings(
 	const FValidateAssetsSettings& InSettings,
 	FValidateAssetsResults& OutResults) const
 {
+	const double StartTime = FPlatformTime::Seconds();
+	EDataValidationResult Result = EDataValidationResult::NotValidated;
+	UE::Telemetry::DataValidation::FValidateAssets Telemetry(InSettings);
+
+	ON_SCOPE_EXIT
+	{
+		Telemetry.SetResults(Result, OutResults);
+		Telemetry.Duration = FPlatformTime::Seconds() - StartTime;
+		FTelemetryRouter::Get().ProvideTelemetry(Telemetry);
+	};
+
 	FMessageLog DataValidationLog(InSettings.MessageLogName);
 	DataValidationLog.SetCurrentPage(InSettings.MessageLogPageTitle);
 
@@ -472,7 +486,7 @@ int32 UEditorValidatorSubsystem::ValidateAssetsWithSettings(
 		NumMessagesPassingFilterPre = DataValidationLog.NumMessages(*Severity);
 	}
 
-	ValidateAssetsInternal(DataValidationLog, TSet<FAssetData>{AssetDataList}, InSettings, OutResults);
+	Result = ValidateAssetsInternal(DataValidationLog, TSet<FAssetData>{AssetDataList}, InSettings, OutResults);
 	
 	if (const EMessageSeverity::Type* Severity = InSettings.ShowMessageLogSeverity.GetPtrOrNull())
 	{
@@ -1031,6 +1045,7 @@ EDataValidationResult UEditorValidatorSubsystem::ValidateChangelistsInternal(
 	const FValidateAssetsSettings& 				Settings,
 	FValidateAssetsResults& 					OutResults) const
 {
+	const double StartTime = FPlatformTime::Seconds();
 	FScopedSlowTask SlowTask(Changelists.Num(), LOCTEXT("DataValidation.ValidatingChangelistTask", "Validating Changelists"));
 	SlowTask.Visibility = ESlowTaskVisibility::Invisible;
 	if (!Settings.bSilent)
@@ -1047,6 +1062,16 @@ EDataValidationResult UEditorValidatorSubsystem::ValidateChangelistsInternal(
 		return EDataValidationResult::NotValidated;
 	}
 
+	EDataValidationResult Result = EDataValidationResult::NotValidated;
+	UE::Telemetry::DataValidation::FValidateAssets Telemetry(Settings);
+
+	ON_SCOPE_EXIT
+	{
+		Telemetry.SetResults(Result, OutResults);
+		Telemetry.Duration = FPlatformTime::Seconds() - StartTime;
+		FTelemetryRouter::Get().ProvideTelemetry(Telemetry);
+	};
+
 	FMessageLog DataValidationLog(Settings.MessageLogName);
 
 	// Choose a specific message log page for this output, flushing in case other recursive calls also write to this log
@@ -1063,14 +1088,14 @@ EDataValidationResult UEditorValidatorSubsystem::ValidateChangelistsInternal(
 		CL->AddToRoot();
 	}
 	
-	ON_SCOPE_EXIT {
+	ON_SCOPE_EXIT 
+	{
 		for (UDataValidationChangelist* CL : Changelists)
 		{
 			CL->RemoveFromRoot();		
 		}
 	};
 
-	EDataValidationResult Result = EDataValidationResult::NotValidated;
 	for (UDataValidationChangelist* Changelist : Changelists)
 	{
 		FText ValidationMessage = FText::Format(LOCTEXT("DataValidation.ValidatingChangelistMessage", "Validating changelist {0}"), Changelist->Description);
diff --git a/Engine/Plugins/Editor/DataValidation/Source/DataValidation/Public/DataValidationTelemetry.h b/Engine/Plugins/Editor/DataValidation/Source/DataValidation/Public/DataValidationTelemetry.h
new file mode 100644
--- /dev/null
+++ b/Engine/Plugins/Editor/DataValidation/Source/DataValidation/Public/DataValidationTelemetry.h
@@ -0,0 +1,66 @@
+// Copyright Epic Games, Inc. All Rights Reserved.
+
+#pragma once
+
+// This header is intended to be included by modules which do not link against DataValidation for access to telemetry data.
+
+#include "Misc/DataValidation.h"
+#include "Misc/Guid.h"
+#include "UObject/UObjectGlobals.h"
+
+struct FValidateAssetsResults;
+struct FValidateAssetsSettings;
+
+namespace UE::Telemetry::DataValidation
+{
+	enum class EValidationMode
+	{
+		Unknown,
+		Editor,
+		Commandlet,
+	};
+
+	struct FValidatorStatistics
+	{
+		int32 AssetsValidated = 0;
+		int32 AssetsAddedForValidation = 0;
+	};
+
+	// Telemetry for an operation where we validate some assets (or some changelists)
+	struct FValidateAssets
+	{
+		static inline constexpr FGuid TelemetryID{0x0e4e8035, 0xdce94c98, 0x94f64b35, 0xef12b049};
+
+		FValidateAssets() = default;
+		FValidateAssets(const FValidateAssetsSettings& InSettings);
+
+		void SetResults(EDataValidationResult InOverallResult, const FValidateAssetsResults& InResults);
+
+		double Duration = 0.0f;
+		EValidationMode Mode = EValidationMode::Unknown;	
+		EDataValidationResult Result = EDataValidationResult::NotValidated;
+		
+		// From FValidateAssetsSettings
+		bool bSkipExcludedDirectories = false;
+		EDataValidationUsecase ValidationUsecase = EDataValidationUsecase::None;
+		bool bLoadAssetsForValidation = false;
+		bool bUnloadAssetsLoadedForValidation = false;
+		bool bLoadExternalObjectsForValidation = false;
+		bool bCaptureAssetLoadLogs = false;
+		bool bCaptureLogsDuringValidation = false;
+		bool bCaptureWarningsDuringValidationAsErrors = false;
+		int32 MaxAssetsToValidate = 0;
+		bool bValidateReferencersOfDeletedAssets = false;
+
+		int NumRequested = 0;
+		int NumExternalObjects = 0;
+		int NumChecked = 0;
+		int NumValid = 0;
+		int NumInvalid = 0;
+		int NumSkipped = 0;
+		int NumWarnings = 0;
+		int NumUnableToValidate = 0;
+		bool bAssetLimitReached = false;
+		TMap<FTopLevelAssetPath, FValidatorStatistics> ValidatorStatistics;
+	};
+};
\ No newline at end of file
diff --git a/Engine/Plugins/Experimental/EditorTelemetry/EditorTelemetry.uplugin b/Engine/Plugins/Experimental/EditorTelemetry/EditorTelemetry.uplugin
--- a/Engine/Plugins/Experimental/EditorTelemetry/EditorTelemetry.uplugin
+++ b/Engine/Plugins/Experimental/EditorTelemetry/EditorTelemetry.uplugin
@@ -20,5 +20,12 @@
 			"Type": "Editor",
 			"LoadingPhase": "PostConfigInit"
 		}
+	],
+	"Plugins" : [
+		{ 
+			"Name": "DataValidation",
+			"Enabled": false,
+			"Optional": true
+		}	
 	]
 }
\ No newline at end of file
diff --git a/Engine/Plugins/Experimental/EditorTelemetry/Source/EditorTelemetry.Build.cs b/Engine/Plugins/Experimental/EditorTelemetry/Source/EditorTelemetry.Build.cs
--- a/Engine/Plugins/Experimental/EditorTelemetry/Source/EditorTelemetry.Build.cs
+++ b/Engine/Plugins/Experimental/EditorTelemetry/Source/EditorTelemetry.Build.cs
@@ -28,6 +28,10 @@ namespace UnrealBuildTool.Rules
 					"ContentBrowserData",
 				}
 			);
+
+			PrivateIncludePathModuleNames.AddRange([
+				"DataValidation"
+			]);
 		}
 	}
 }
diff --git a/Engine/Plugins/Experimental/EditorTelemetry/Source/Private/EditorTelemetry.cpp b/Engine/Plugins/Experimental/EditorTelemetry/Source/Private/EditorTelemetry.cpp
--- a/Engine/Plugins/Experimental/EditorTelemetry/Source/Private/EditorTelemetry.cpp
+++ b/Engine/Plugins/Experimental/EditorTelemetry/Source/Private/EditorTelemetry.cpp
@@ -2,12 +2,14 @@
 
 #include "EditorTelemetry.h"
 
+#include "Algo/Copy.h"
 #include "AnalyticsTracer.h"
 #include "AssetRegistry/AssetRegistryTelemetry.h"
 #include "CollectionManagerModule.h"
 #include "ContentBrowserModule.h"
 #include "ContentBrowserTelemetry.h"
 #include "CookOnTheSide/CookOnTheFlyServer.h"
+#include "DataValidationTelemetry.h"
 #include "DerivedDataCacheInterface.h"
 #include "DerivedDataCacheUsageStats.h"
 #include "Editor.h"
@@ -25,6 +27,7 @@
 #include "StudioTelemetry.h"
 #include "Subsystems/AssetEditorSubsystem.h"
 #include "TelemetryRouter.h"
+#include "Templates/Greater.h"
 #include "UObject/ICookInfo.h"
 #include "UnrealEdGlobals.h"
 #include "Virtualization/VirtualizationSystem.h"
@@ -44,6 +47,45 @@ namespace Private
 	};
 }
 
+namespace UE::Telemetry::DataValidation
+{
+	inline const TCHAR* LexToString(EDataValidationResult Result)
+	{
+		switch (Result)
+		{
+			default:
+			case EDataValidationResult::NotValidated: return TEXT("NotValidated");
+			case EDataValidationResult::Invalid: return TEXT("Invalid");
+			case EDataValidationResult::Valid: return TEXT("Valid");
+		}
+	}
+
+	inline const TCHAR* LexToString(EDataValidationUsecase Usecase)
+	{
+		switch (Usecase)
+		{
+			default:
+			case EDataValidationUsecase::None: return TEXT("None");
+			case EDataValidationUsecase::Manual: return TEXT("Manual");
+			case EDataValidationUsecase::Commandlet: return TEXT("Commandlet");
+			case EDataValidationUsecase::Save: return TEXT("Save");
+			case EDataValidationUsecase::PreSubmit: return TEXT("PreSubmit");
+			case EDataValidationUsecase::Script: return TEXT("Script");
+		}
+	}
+
+	inline const TCHAR* LexToString(EValidationMode Mode)
+	{
+		switch (Mode)
+		{
+			default:
+			case EValidationMode::Unknown: return TEXT("Unknown");
+			case EValidationMode::Editor: return TEXT("Editor");
+			case EValidationMode::Commandlet: return TEXT("Commandlet");
+		}
+	}
+}
+
 const TCHAR* LexToString(ECollectionTelemetryAssetAddedWorkflow Enum)
 {
 	switch(Enum)
@@ -1094,6 +1136,61 @@ void FEditorTelemetry::StartSession()
 			});
 	}
 
+	// Data validation telemetry
+	{
+		using namespace UE::Telemetry;
+		using namespace UE::Telemetry::DataValidation;
+
+		Router.OnTelemetry<FValidateAssets>([this](const FValidateAssets& Data){
+			FString ValidatorStats;
+			{
+				::Private::FAnalyticsJsonWriter J(&ValidatorStats);
+				J.WriteObjectStart();
+				TArray<TPair<FTopLevelAssetPath, FValidatorStatistics>> Stats;
+				Algo::Copy(Data.ValidatorStatistics, Stats);
+				Algo::SortBy(Stats, 
+					[](const TPair<FTopLevelAssetPath, FValidatorStatistics>& Pair) { return Pair.Value.AssetsValidated + Pair.Value.AssetsAddedForValidation; },
+					TGreater<int32>{});
+				for (const TPair<FTopLevelAssetPath, FValidatorStatistics>& Pair : MakeArrayView(Stats).Left(10))
+				{
+					J.WriteObjectStart(WriteToString<256>(Pair.Key).ToView());
+					J.WriteValue(TEXT("NumAssetsValidated"), Pair.Value.AssetsValidated);
+					J.WriteValue(TEXT("NumAssetsAddedForValidation"), Pair.Value.AssetsAddedForValidation);
+					J.WriteObjectEnd();
+				}
+				J.WriteObjectEnd();
+				J.Close();
+			}
+			FStudioTelemetry::Get().RecordEvent(TEXT("Editor.DataValidation.ValidateAssets"), 
+			{
+				{ TEXT("SchemaVersion"), 1 },
+				{ TEXT("Duration"), Data.Duration },
+				{ TEXT("Mode"), LexToString(Data.Mode) },
+				{ TEXT("Result"), LexToString(Data.Result) },
+				{ TEXT("SkipExcludedDirectories"), Data.bSkipExcludedDirectories },
+				{ TEXT("ValidationUsecase"), LexToString(Data.ValidationUsecase) },
+				{ TEXT("LoadAssetsForValidation"), Data.bLoadAssetsForValidation },
+				{ TEXT("UnloadAssetsLoadedForValidation"), Data.bUnloadAssetsLoadedForValidation },
+				{ TEXT("LoadExternalObjectsForValidation"), Data.bLoadExternalObjectsForValidation },
+				{ TEXT("CaptureAssetLoadLogs"), Data.bCaptureAssetLoadLogs },
+				{ TEXT("CaptureLogsDuringValidation"), Data.bCaptureLogsDuringValidation },
+				{ TEXT("CaptureWarningsDuringValidationAsErrors"), Data.bCaptureWarningsDuringValidationAsErrors },
+				{ TEXT("MaxAssetsToValidate"), Data.MaxAssetsToValidate },
+				{ TEXT("ValidateReferencersOfDeletedAssets"), Data.bValidateReferencersOfDeletedAssets },
+				{ TEXT("NumRequested"), Data.NumRequested },
+				{ TEXT("NumExternalObjects"), Data.NumExternalObjects },
+				{ TEXT("NumChecked"), Data.NumChecked },
+				{ TEXT("NumValid"), Data.NumValid },
+				{ TEXT("NumInvalid"), Data.NumInvalid },
+				{ TEXT("NumSkipped"), Data.NumSkipped },
+				{ TEXT("NumWarnings"), Data.NumWarnings },
+				{ TEXT("NumUnableToValidate"), Data.NumUnableToValidate },
+				{ TEXT("AssetLimitReached"), Data.bAssetLimitReached },
+				{ TEXT("ValidatorStatistics"), ValidatorStats },
+			});
+		});
+	}
+
 	{
 		UE::Virtualization::GetAnalyticsRecordEvent().AddLambda([](const FString& EventName, const TArray<FAnalyticsEventAttribute>& Attributes, UE::Virtualization::EAnalyticsFlags Flags)
 			{
