Fix ValidatorStatistics field of Editor.DataValidation.ValidateAsset to be an FJsonFragment to embed correctly as an object in the final JSON

diff --git a/Engine/Plugins/Experimental/EditorTelemetry/Source/Private/EditorTelemetry.cpp b/Engine/Plugins/Experimental/EditorTelemetry/Source/Private/EditorTelemetry.cpp
--- a/Engine/Plugins/Experimental/EditorTelemetry/Source/Private/EditorTelemetry.cpp
+++ b/Engine/Plugins/Experimental/EditorTelemetry/Source/Private/EditorTelemetry.cpp
@@ -1186,7 +1186,7 @@ void FEditorTelemetry::StartSession()
 				{ TEXT("NumWarnings"), Data.NumWarnings },
 				{ TEXT("NumUnableToValidate"), Data.NumUnableToValidate },
 				{ TEXT("AssetLimitReached"), Data.bAssetLimitReached },
-				{ TEXT("ValidatorStatistics"), ValidatorStats },
+				{ TEXT("ValidatorStatistics"), FJsonFragment(MoveTemp(ValidatorStats)) },
 			});
 		});
 	}
