fix: schema version should be outside of the stats data

#Synced-CL 49027723

diff --git a/Engine/Plugins/Experimental/EditorTelemetry/Source/Private/EditorTelemetry.cpp b/Engine/Plugins/Experimental/EditorTelemetry/Source/Private/EditorTelemetry.cpp
--- a/Engine/Plugins/Experimental/EditorTelemetry/Source/Private/EditorTelemetry.cpp
+++ b/Engine/Plugins/Experimental/EditorTelemetry/Source/Private/EditorTelemetry.cpp
@@ -125,8 +125,6 @@ void FEditorTelemetry::RecordEvent_Cooking(TArray<FAnalyticsEventAttribute> Attr
 
 	const int SchemaVersion = 4;
 
-	Attributes.Emplace(TEXT("SchemaVersion"), SchemaVersion);
-
 	TMap<FString,FAnalyticsEventAttribute> CookAttributes;
 	
 	// Sends each cook stat to the studio analytics system.
@@ -204,7 +202,9 @@ void FEditorTelemetry::RecordEvent_Cooking(TArray<FAnalyticsEventAttribute> Attr
 
 	FString CookStatsJsonStr;
 	Private::FAnalyticsJsonWriter CookStatsJsonWriter(&CookStatsJsonStr);
-	CookStatsJsonWriter.WriteArrayStart();
+	CookStatsJsonWriter.WriteObjectStart();
+	CookStatsJsonWriter.WriteValue("SchemaVersion", SchemaVersion);
+	CookStatsJsonWriter.WriteArrayStart("ProfileData");
 	const FDecimalNumberFormattingRules FormattingRules;
 	for (auto& It : Attributes)
 	{
@@ -225,6 +225,7 @@ void FEditorTelemetry::RecordEvent_Cooking(TArray<FAnalyticsEventAttribute> Attr
 		CookStatsJsonWriter.WriteObjectEnd();
 	}
 	CookStatsJsonWriter.WriteArrayEnd();
+	CookStatsJsonWriter.WriteObjectEnd();
 	CookStatsJsonWriter.Close();
 
 	FStudioTelemetry::Get().RecordEvent(TEXT("Core.Cooking"), 
