Undo cook event telemetry changes

#Synced-CL 49155865

diff --git a/Engine/Plugins/Experimental/EditorTelemetry/Source/Private/EditorTelemetry.cpp b/Engine/Plugins/Experimental/EditorTelemetry/Source/Private/EditorTelemetry.cpp
--- a/Engine/Plugins/Experimental/EditorTelemetry/Source/Private/EditorTelemetry.cpp
+++ b/Engine/Plugins/Experimental/EditorTelemetry/Source/Private/EditorTelemetry.cpp
@@ -16,7 +16,6 @@
 #include "Editor/UnrealEdEngine.h"
 #include "Engine/AssetManager.h"
 #include "Experimental/ZenServerInterface.h"
-#include "Internationalization/FastDecimalFormat.h"
 #include "FileHelpers.h"
 #include "HttpManager.h"
 #include "HttpModule.h"
@@ -32,7 +31,6 @@
 #include "UObject/ICookInfo.h"
 #include "UnrealEdGlobals.h"
 #include "Virtualization/VirtualizationSystem.h"
-#include "Cooker/CookProfiling.h"
 
 namespace Private
 {
@@ -123,36 +121,15 @@ void FEditorTelemetry::RecordEvent_Cooking(TArray<FAnalyticsEventAttribute> Attr
 {
 #if ENABLE_COOK_STATS
 
-	const int SchemaVersion = 4;
+	const int SchemaVersion = 3;
+
+	Attributes.Emplace(TEXT("SchemaVersion"), SchemaVersion);
 
 	TMap<FString,FAnalyticsEventAttribute> CookAttributes;
 	
 	// Sends each cook stat to the studio analytics system.
 	auto GatherAnalyticsAttributes = [&CookAttributes](const FString& StatName, const TArray<FCookStatsManager::StringKeyValue>& StatAttributes)
 	{
-		//Special handling for the Hierarchical profiling data
-		if (StatName == TEXT("Cook.Profile"))
-		{
-			// Hierarchical stats have 2 attributes per stat. 
-			// The path (signifies the hierarchy information) and the stat itself
-			if (StatAttributes.Num() != 2)
-			{
-				return;
-			}
-
-			// When generating the telemetry event we don't care about the hierarchy but only stat itself.
-			const FCookStatsManager::StringKeyValue& Stat = StatAttributes[1];
-
-			const FString FormattedAttrName = (StatName + "_" + Stat.Key).Replace(TEXT("."), TEXT("_"));
-
-			if (CookAttributes.Find(FormattedAttrName) == nullptr)
-			{
-				CookAttributes.Emplace(FormattedAttrName, Stat.Value.IsNumeric() ? FAnalyticsEventAttribute(FormattedAttrName, FCString::Atof(*Stat.Value)) : FAnalyticsEventAttribute(FormattedAttrName, Stat.Value));
-			}
-
-			return;
-		}
-
 		for (const auto& Attr : StatAttributes)
 		{
 			const FString FormattedAttrName = (StatName + "_" + Attr.Key).Replace(TEXT("."), TEXT("_"));
@@ -173,7 +150,25 @@ void FEditorTelemetry::RecordEvent_Cooking(TArray<FAnalyticsEventAttribute> Attr
 		Attributes.Emplace((*it).Value);
 	}
 
-	CookProfiling::GatherAccumulatedDDCStats(Attributes);
+	// Gather the DDC summary stats
+	FDerivedDataCacheSummaryStats SummaryStats;
+
+	GatherDerivedDataCacheSummaryStats(SummaryStats);
+
+	// Append to the attributes
+	for (const FDerivedDataCacheSummaryStat& Stat : SummaryStats.Stats)
+	{
+		FString AttributeName = TEXT("DDC_Summary") + Stat.Key.Replace(TEXT("."), TEXT("_"));
+
+		if (Stat.Value.IsNumeric())
+		{
+			Attributes.Emplace(AttributeName, FCString::Atof(*Stat.Value));
+		}
+		else
+		{
+			Attributes.Emplace(AttributeName, Stat.Value);
+		}
+	}
 
 #if UE_WITH_ZEN
 	// Gather Zen analytics
@@ -191,45 +186,7 @@ void FEditorTelemetry::RecordEvent_Cooking(TArray<FAnalyticsEventAttribute> Attr
 
 	FShaderStatsFunctions::GatherShaderAnalytics(Attributes);
 	
-	// Format the attribs into an easily digestible JSON format that splits each attribute into 
-	// a separate json object
-	// Example:
-	// [
-	// { "name": "PackTime", "value" : 13.34},
-	// { "name": "MaterialShader_Count", "value" : 12 },
-	// ...
-	// ]
-
-	FString CookStatsJsonStr;
-	Private::FAnalyticsJsonWriter CookStatsJsonWriter(&CookStatsJsonStr);
-	CookStatsJsonWriter.WriteObjectStart();
-	CookStatsJsonWriter.WriteValue("SchemaVersion", SchemaVersion);
-	CookStatsJsonWriter.WriteArrayStart("ProfileData");
-	const FDecimalNumberFormattingRules FormattingRules;
-	for (auto& It : Attributes)
-	{
-		CookStatsJsonWriter.WriteObjectStart();
-		CookStatsJsonWriter.WriteValue("name", It.GetName());
-		const FString& StringValue = It.GetValue();
-
-		if (StringValue.IsNumeric())
-		{
-			double ValueAsDouble = 0;
-			LexFromString(ValueAsDouble, StringValue);
-			CookStatsJsonWriter.WriteValue("value", ValueAsDouble);
-		}
-		else
-		{
-			CookStatsJsonWriter.WriteValue("value", StringValue);
-		}
-		CookStatsJsonWriter.WriteObjectEnd();
-	}
-	CookStatsJsonWriter.WriteArrayEnd();
-	CookStatsJsonWriter.WriteObjectEnd();
-	CookStatsJsonWriter.Close();
-
-	FStudioTelemetry::Get().RecordEvent(TEXT("Core.Cooking"), 
-		{ { TEXT("CookStats"), FJsonFragment(MoveTemp(CookStatsJsonStr)) } });
+	FStudioTelemetry::Get().RecordEvent(TEXT("Core.Cooking"), Attributes);
 #endif
 }
 
@@ -953,6 +910,8 @@ void FEditorTelemetry::StartSession()
 					return;
 				}
 				TArray<FAnalyticsEventAttribute> Attributes;
+				Attributes.Emplace(TEXT("MapName"), EditorMapName);
+
 				FEditorTelemetry::RecordEvent_Cooking(Attributes);
 				FEditorTelemetry::RecordEvent_CoreSystems(TEXT("Cooking"), Attributes);
 				FEditorTelemetry::RecordEvent_DDCVerification(TEXT("Cooking"), Attributes);
@@ -969,6 +928,8 @@ void FEditorTelemetry::StartSession()
 		CookingSpan = FStudioTelemetry::Get().StartSpan(TEXT("Cooking"));
 
 		TArray<FAnalyticsEventAttribute> Attributes;
+		Attributes.Emplace(TEXT("MapName"), EditorMapName);
+
 		CookingSpan->AddAttributes(Attributes);
 	});
 
