A refactor of Editor Telemetry events to reduce duplication of attribute data. As a result events are much smaller.

We should be using Database joins from now on to add additonal contexts to a given event.

Added the functionality to push Span GUIDs into the EditorTelemetry events. This will allow us to join spans to detailed event information.

Prepended the Content Worker default attributes so they don't clash with default attributes in StudioTelemetry

#Synced-CL 49803489

diff --git a/Engine/Plugins/Experimental/EditorTelemetry/Source/Private/EditorTelemetry.cpp b/Engine/Plugins/Experimental/EditorTelemetry/Source/Private/EditorTelemetry.cpp
--- a/Engine/Plugins/Experimental/EditorTelemetry/Source/Private/EditorTelemetry.cpp
+++ b/Engine/Plugins/Experimental/EditorTelemetry/Source/Private/EditorTelemetry.cpp
@@ -117,13 +117,14 @@ FEditorTelemetry& FEditorTelemetry::Get()
 	return StudioTelemetryEditorInstance;
 }
 
-void FEditorTelemetry::RecordEvent_Cooking(TArray<FAnalyticsEventAttribute> Attributes)
+void FEditorTelemetry::RecordEvent_Cooking(const FString& Context, TArray<FAnalyticsEventAttribute> Attributes)
 {
 #if ENABLE_COOK_STATS
 
-	const int SchemaVersion = 3;
+	const int SchemaVersion = 4;
 
 	Attributes.Emplace(TEXT("SchemaVersion"), SchemaVersion);
+	Attributes.Emplace(TEXT("Context"), Context);
 
 	TMap<FString,FAnalyticsEventAttribute> CookAttributes;
 	
@@ -150,92 +151,20 @@ void FEditorTelemetry::RecordEvent_Cooking(TArray<FAnalyticsEventAttribute> Attr
 		Attributes.Emplace((*it).Value);
 	}
 
-	// Gather the DDC summary stats
-	FDerivedDataCacheSummaryStats SummaryStats;
-
-	GatherDerivedDataCacheSummaryStats(SummaryStats);
-
-	// Append to the attributes
-	for (const FDerivedDataCacheSummaryStat& Stat : SummaryStats.Stats)
-	{
-		FString AttributeName = TEXT("DDC_Summary") + Stat.Key.Replace(TEXT("."), TEXT("_"));
-
-		if (Stat.Value.IsNumeric())
-		{
-			Attributes.Emplace(AttributeName, FCString::Atof(*Stat.Value));
-		}
-		else
-		{
-			Attributes.Emplace(AttributeName, Stat.Value);
-		}
-	}
-
-#if UE_WITH_ZEN
-	// Gather Zen analytics
-	if (UE::Zen::IsDefaultServicePresent())
-	{
-		UE::Zen::GetDefaultServiceInstance().GatherAnalytics(Attributes);
-	}
-#endif
-
-	if (UE::Virtualization::IVirtualizationSystem::Get().IsEnabled())
-	{
-		// Gather Virtualization analytics
-		UE::Virtualization::IVirtualizationSystem::Get().GatherAnalytics(Attributes);
-	}
-
-	FShaderStatsFunctions::GatherShaderAnalytics(Attributes);
-	
 	FStudioTelemetry::Get().RecordEvent(TEXT("Core.Cooking"), Attributes);
+
 #endif
 }
 
 void FEditorTelemetry::RecordEvent_Loading(const FString& Context, double LoadingSeconds, TArray<FAnalyticsEventAttribute> Attributes )
 {
-	const int SchemaVersion = 4;
+	const int SchemaVersion = 5;
 
 	Attributes.Emplace(TEXT("SchemaVersion"), SchemaVersion);
 	Attributes.Emplace(TEXT("Context"), Context);
 	Attributes.Emplace(TEXT("LoadingName"), Context);
 	Attributes.Emplace(TEXT("LoadingSeconds"), LoadingSeconds);
 	
-#if ENABLE_COOK_STATS
-
-#if UE_WITH_ZEN
-	// Gather Zen analytics
-	if (UE::Zen::IsDefaultServicePresent())
-	{
-		UE::Zen::GetDefaultServiceInstance().GatherAnalytics(Attributes);
-	}
-#endif
-
-	if (UE::Virtualization::IVirtualizationSystem::Get().IsEnabled())
-	{
-		// Gather Virtualization analytics
-		UE::Virtualization::IVirtualizationSystem::Get().GatherAnalytics(Attributes);
-	}
-
-	// Gather the DDC summary stats
-	FDerivedDataCacheSummaryStats SummaryStats;
-
-	GatherDerivedDataCacheSummaryStats(SummaryStats);
-
-	// Append to the attributes
-	for (const FDerivedDataCacheSummaryStat& Stat : SummaryStats.Stats)
-	{
-		FString AttributeName = TEXT("DDC_Summary_") + Stat.Key.Replace(TEXT("."), TEXT("_"));
-
-		if (Stat.Value.IsNumeric())
-		{
-			Attributes.Emplace(AttributeName, FCString::Atof(*Stat.Value));
-		}
-		else
-		{
-			Attributes.Emplace(AttributeName, Stat.Value);
-		}
-	}
-#endif			
-
 	FStudioTelemetry::Get().RecordEvent(TEXT("Core.Loading"), Attributes);
 }
 
@@ -383,6 +312,18 @@ void FEditorTelemetry::RecordEvent_VirtualAssets(const FString& Context, TArray<
 	}
 }
 
+void FEditorTelemetry::RecordEvent_ShaderStats(const FString& Context, TArray<FAnalyticsEventAttribute> Attributes)
+{
+	const int SchemaVersion = 1;
+
+	Attributes.Emplace(TEXT("SchemaVersion"), SchemaVersion);
+	Attributes.Emplace(TEXT("Context"), Context);
+
+	FShaderStatsFunctions::GatherShaderAnalytics(Attributes);
+
+	FStudioTelemetry::Get().RecordEvent(TEXT("Core.ShaderStats"), Attributes);
+}
+
 void FEditorTelemetry::RecordEvent_MemoryLLM(const FString& Context, TArray<FAnalyticsEventAttribute> Attributes)
 {
 #if ENABLE_LOW_LEVEL_MEM_TRACKER
@@ -912,9 +853,10 @@ void FEditorTelemetry::StartSession()
 				TArray<FAnalyticsEventAttribute> Attributes;
 				Attributes.Emplace(TEXT("MapName"), EditorMapName);
 
-				FEditorTelemetry::RecordEvent_Cooking(Attributes);
-				FEditorTelemetry::RecordEvent_CoreSystems(TEXT("Cooking"), Attributes);
-				FEditorTelemetry::RecordEvent_DDCVerification(TEXT("Cooking"), Attributes);
+				FEditorTelemetry::RecordEvent_Cooking(TEXT("Cook.End"), Attributes);
+				FEditorTelemetry::RecordEvent_CoreSystems(TEXT("Cook.End"), Attributes);
+				FEditorTelemetry::RecordEvent_DDCVerification(TEXT("Cook.End"), Attributes);
+				FEditorTelemetry::RecordEvent_ShaderStats(TEXT("Cook.End"), CookingSpan->GetAttributes());
 			});
 	}
 
@@ -931,6 +873,9 @@ void FEditorTelemetry::StartSession()
 		Attributes.Emplace(TEXT("MapName"), EditorMapName);
 
 		CookingSpan->AddAttributes(Attributes);
+
+		const double TimeToCook = FPlatformTime::Seconds() - SessionStartTime;
+		FEditorTelemetry::RecordEvent_Loading(TEXT("Cook.Start"), TimeToCook, CookingSpan->GetAttributes());
 	});
 
 	UE::Cook::FDelegates::CookFinished.AddLambda([this](UE::Cook::ICookInfo& CookInfo)
@@ -947,14 +892,15 @@ void FEditorTelemetry::StartSession()
 		if (bIsCookWorker)
 		{
 			// Cook-worker specific telemetry goes here
-			FEditorTelemetry::RecordEvent_DDCVerification(TEXT("Cooking"), CookingSpan->GetAttributes());
+			FEditorTelemetry::RecordEvent_DDCVerification(TEXT("Cook.End"), CookingSpan->GetAttributes());
 		}
 		else
 		{
 			// Cook director/single-process cook telemetry goes here
-			FEditorTelemetry::RecordEvent_Cooking(CookingSpan->GetAttributes());
-			FEditorTelemetry::RecordEvent_CoreSystems(TEXT("Cooking"), CookingSpan->GetAttributes());
-			FEditorTelemetry::RecordEvent_DDCVerification(TEXT("Cooking"), CookingSpan->GetAttributes());
+			FEditorTelemetry::RecordEvent_Cooking(TEXT("Cook.End"), CookingSpan->GetAttributes());
+			FEditorTelemetry::RecordEvent_CoreSystems(TEXT("Cook.End"), CookingSpan->GetAttributes());
+			FEditorTelemetry::RecordEvent_DDCVerification(TEXT("Cook.End"), CookingSpan->GetAttributes());
+			FEditorTelemetry::RecordEvent_ShaderStats(TEXT("Cook.End"), CookingSpan->GetAttributes());
 
 			// End the cooking span
 			FStudioTelemetry::Get().EndSpan(CookingSpan);
@@ -1212,4 +1158,5 @@ void FEditorTelemetry::EndSession()
 		FStudioTelemetry::Get().EndSpan(EditorSpan);
 		EditorSpan.Reset();
 	}
+
 }
diff --git a/Engine/Plugins/Experimental/EditorTelemetry/Source/Public/EditorTelemetry.h b/Engine/Plugins/Experimental/EditorTelemetry/Source/Public/EditorTelemetry.h
--- a/Engine/Plugins/Experimental/EditorTelemetry/Source/Public/EditorTelemetry.h
+++ b/Engine/Plugins/Experimental/EditorTelemetry/Source/Public/EditorTelemetry.h
@@ -32,7 +32,7 @@ public:
 	UE_API void EndSession();
 
 	/** Useful event recording functions */
-	UE_API void RecordEvent_Cooking(TArray<FAnalyticsEventAttribute> Attributes = {});
+	UE_API void RecordEvent_Cooking(const FString& Context, TArray<FAnalyticsEventAttribute> Attributes = {});
 	UE_API void RecordEvent_Loading(const FString& Context, double LoadingSeconds, TArray<FAnalyticsEventAttribute> Attributes = {});
 	UE_API void RecordEvent_CoreSystems(const FString& Context, TArray<FAnalyticsEventAttribute> Attributes = {});
 	UE_API void RecordEvent_DDCResource(const FString& Context, TArray<FAnalyticsEventAttribute> Attributes = {});
@@ -41,6 +41,8 @@ public:
 	UE_API void RecordEvent_Zen(const FString& Context, TArray<FAnalyticsEventAttribute> Attributes = {});
 	UE_API void RecordEvent_VirtualAssets(const FString& Context, TArray<FAnalyticsEventAttribute> Attributes = {});
 	UE_API void RecordEvent_MemoryLLM(const FString& Context, TArray<FAnalyticsEventAttribute> Attributes = {});
+	UE_API void RecordEvent_ShaderStats(const FString& Context, TArray<FAnalyticsEventAttribute> Attributes = {});
+
 	UE_API void RegisterCollectionWorkflowDelegates(FTelemetryRouter& Router);
 
 private:	
diff --git a/Engine/Source/Runtime/Analytics/Analytics/Private/AnalyticsTracer.cpp b/Engine/Source/Runtime/Analytics/Analytics/Private/AnalyticsTracer.cpp
--- a/Engine/Source/Runtime/Analytics/Analytics/Private/AnalyticsTracer.cpp
+++ b/Engine/Source/Runtime/Analytics/Analytics/Private/AnalyticsTracer.cpp
@@ -62,6 +62,9 @@ void FAnalyticsSpan::Start(const TArray<FAnalyticsEventAttribute>& AdditionalAtt
 	EndTime			= FDateTime::UtcNow();
 	Duration		= 0;	
 	IsActive		= true;
+
+	// Always add the Span GUID to the attributes so we can join this span to other events
+	Attributes.Add(FAnalyticsEventAttribute(TEXT("Span_GUID"), Guid.ToString()));
 }
 
 bool FAnalyticsSpan::GetIsActive() const
@@ -91,7 +94,6 @@ void FAnalyticsSpan::End(const TArray<FAnalyticsEventAttribute>& AdditionalAttri
 
 	EventAttributes.Add(FAnalyticsEventAttribute(TEXT("SchemaVersion"), SpanSchemaVersion));
 	EventAttributes.Add(FAnalyticsEventAttribute(TEXT("Span_Name"), Name.ToString()));
-	EventAttributes.Add(FAnalyticsEventAttribute(TEXT("Span_GUID"), Guid.ToString()));
 	EventAttributes.Add(FAnalyticsEventAttribute(TEXT("Span_ThreadId"), ThreadId));
 	EventAttributes.Add(FAnalyticsEventAttribute(TEXT("Span_Depth"), StackDepth));
 	EventAttributes.Add(FAnalyticsEventAttribute(TEXT("Span_StartUTC"), StartTime.ToUnixTimestampDecimal()));
