Added TicketId to AnalyticsContentWorker
Added smaller Core.Cooking.Summary event to EditorTelemetry
Added extra checks not to send incomplete Horde environment on ContentWorkers
Removed deprecated  GatherAnalytics function in DerivedDataCacheInterface

#Synced-CL 50046023

diff --git a/Engine/Plugins/Experimental/EditorTelemetry/Source/Private/EditorTelemetry.cpp b/Engine/Plugins/Experimental/EditorTelemetry/Source/Private/EditorTelemetry.cpp
--- a/Engine/Plugins/Experimental/EditorTelemetry/Source/Private/EditorTelemetry.cpp
+++ b/Engine/Plugins/Experimental/EditorTelemetry/Source/Private/EditorTelemetry.cpp
@@ -117,15 +117,10 @@ FEditorTelemetry& FEditorTelemetry::Get()
 	return StudioTelemetryEditorInstance;
 }
 
-void FEditorTelemetry::RecordEvent_Cooking(const FString& Context, TArray<FAnalyticsEventAttribute> Attributes)
+void FEditorTelemetry::RecordEvent_Cooking(const FString& Context, const TArray<FAnalyticsEventAttribute>& Attributes)
 {
 #if ENABLE_COOK_STATS
 
-	const int SchemaVersion = 4;
-
-	Attributes.Emplace(TEXT("SchemaVersion"), SchemaVersion);
-	Attributes.Emplace(TEXT("Context"), Context);
-
 	TMap<FString,FAnalyticsEventAttribute> CookAttributes;
 	
 	// Sends each cook stat to the studio analytics system.
@@ -145,30 +140,68 @@ void FEditorTelemetry::RecordEvent_Cooking(const FString& Context, TArray<FAnaly
 	// Now actually grab the stats 
 	FCookStatsManager::LogCookStats(GatherAnalyticsAttributes);
 
-	// Add the values to the attributes
-	for (TMap<FString, FAnalyticsEventAttribute>::TConstIterator it(CookAttributes); it; ++it)
+	// Build and send the full Cook event
 	{
-		Attributes.Emplace((*it).Value);
+		const int SchemaVersion = 4;
+
+		TArray<FAnalyticsEventAttribute> FullEventAttributes = Attributes;
+
+		FullEventAttributes.Emplace(TEXT("SchemaVersion"), SchemaVersion);
+		FullEventAttributes.Emplace(TEXT("Context"), Context);
+
+		// Add the values to the full cook event attributes
+		for (TMap<FString, FAnalyticsEventAttribute>::TConstIterator it(CookAttributes); it; ++it)
+		{
+			FullEventAttributes.Emplace((*it).Value);
+		}
+
+		// Record the full cook event
+		FStudioTelemetry::Get().RecordEvent(TEXT("Core.Cooking"), FullEventAttributes);
 	}
 
-	FStudioTelemetry::Get().RecordEvent(TEXT("Core.Cooking"), Attributes);
+	// Build and send the summarized Cook event
+	{
+		const int SchemaVersion = 1;
+
+		TArray<FAnalyticsEventAttribute> SummaryEventAttributes = Attributes;
+
+		SummaryEventAttributes.Emplace(TEXT("SchemaVersion"), SchemaVersion);
+		SummaryEventAttributes.Emplace(TEXT("Context"), Context);
+
+		// Add the values to the summary cook event attributes
+		for (TMap<FString, FAnalyticsEventAttribute>::TConstIterator it(CookAttributes); it; ++it)
+		{
+			const FString& KeyName = (*it).Key;
+
+			// Filter the Cook attributes to the ones we want to send in the summary event
+			if ( KeyName.Find("Cook_Profile_") != INDEX_NONE || KeyName.Find("CookOnTheFlyServer_") != INDEX_NONE || KeyName.Find("Package_") != INDEX_NONE || KeyName.Find("UnversionedProperties_") != INDEX_NONE )
+			{
+				SummaryEventAttributes.Emplace((*it).Value);
+			}
+		}
+
+		// Record the summary cook event
+		FStudioTelemetry::Get().RecordEvent(TEXT("Core.Cooking.Summary"), SummaryEventAttributes);
+	}
 
 #endif
 }
 
-void FEditorTelemetry::RecordEvent_Loading(const FString& Context, double LoadingSeconds, TArray<FAnalyticsEventAttribute> Attributes )
+void FEditorTelemetry::RecordEvent_Loading(const FString& Context, double LoadingSeconds, const TArray<FAnalyticsEventAttribute>& Attributes)
 {
 	const int SchemaVersion = 5;
 
-	Attributes.Emplace(TEXT("SchemaVersion"), SchemaVersion);
-	Attributes.Emplace(TEXT("Context"), Context);
-	Attributes.Emplace(TEXT("LoadingName"), Context);
-	Attributes.Emplace(TEXT("LoadingSeconds"), LoadingSeconds);
+	TArray<FAnalyticsEventAttribute> EventAttributes = Attributes;
+
+	EventAttributes.Emplace(TEXT("SchemaVersion"), SchemaVersion);
+	EventAttributes.Emplace(TEXT("Context"), Context);
+	EventAttributes.Emplace(TEXT("LoadingName"), Context);
+	EventAttributes.Emplace(TEXT("LoadingSeconds"), LoadingSeconds);
 	
-	FStudioTelemetry::Get().RecordEvent(TEXT("Core.Loading"), Attributes);
+	FStudioTelemetry::Get().RecordEvent(TEXT("Core.Loading"), EventAttributes);
 }
 
-void FEditorTelemetry::RecordEvent_DDCResource(const FString& Context, TArray<FAnalyticsEventAttribute> Attributes)
+void FEditorTelemetry::RecordEvent_DDCResource(const FString& Context, const TArray<FAnalyticsEventAttribute>& Attributes)
 {
 #if ENABLE_COOK_STATS
 	// Gather the latest resource stats
@@ -178,9 +211,6 @@ void FEditorTelemetry::RecordEvent_DDCResource(const FString& Context, TArray<FA
 
 	const int SchemaVersion = 4;
 
-	Attributes.Emplace(TEXT("SchemaVersion"), SchemaVersion);
-	Attributes.Emplace(TEXT("Context"), Context);
-	
 	// Send a resource event per asset type
 	for (const FDerivedDataCacheResourceStat& Stat : ResourceStats)
 	{
@@ -195,6 +225,8 @@ void FEditorTelemetry::RecordEvent_DDCResource(const FString& Context, TArray<FA
 	
 		TArray<FAnalyticsEventAttribute> EventAttributes = Attributes;
 
+		EventAttributes.Emplace(TEXT("SchemaVersion"), SchemaVersion);
+		EventAttributes.Emplace(TEXT("Context"), Context);
 		EventAttributes.Emplace(TEXT("AssetType"), Stat.AssetType);
 		EventAttributes.Emplace(TEXT("Load_Count"), Stat.LoadCount);
 		EventAttributes.Emplace(TEXT("Load_TimeSec"), Stat.LoadTimeSec);
@@ -214,13 +246,15 @@ void FEditorTelemetry::RecordEvent_DDCResource(const FString& Context, TArray<FA
 }
 
 
-void FEditorTelemetry::RecordEvent_DDCSummary(const FString& Context, TArray<FAnalyticsEventAttribute> Attributes)
+void FEditorTelemetry::RecordEvent_DDCSummary(const FString& Context, const TArray<FAnalyticsEventAttribute>& Attributes)
 {
 #if ENABLE_COOK_STATS
 	const int SchemaVersion = 4;
 
-	Attributes.Emplace(TEXT("SchemaVersion"), SchemaVersion);
-	Attributes.Emplace(TEXT("Context"), Context);
+	TArray<FAnalyticsEventAttribute> EventAttributes = Attributes;
+
+	EventAttributes.Emplace(TEXT("SchemaVersion"), SchemaVersion);
+	EventAttributes.Emplace(TEXT("Context"), Context);
 
 	// Gather the summary stats
 	FDerivedDataCacheSummaryStats SummaryStats;
@@ -234,23 +268,21 @@ void FEditorTelemetry::RecordEvent_DDCSummary(const FString& Context, TArray<FAn
 
 		if (Stat.Value.IsNumeric())
 		{
-			Attributes.Emplace(AttributeName, FCString::Atof(*Stat.Value));
+			EventAttributes.Emplace(AttributeName, FCString::Atof(*Stat.Value));
 		}
 		else
 		{
-			Attributes.Emplace(AttributeName, Stat.Value);
+			EventAttributes.Emplace(AttributeName, Stat.Value);
 		}
 	}
 
-	FStudioTelemetry::Get().RecordEvent(TEXT("Core.DDC.Summary"), Attributes);
+	FStudioTelemetry::Get().RecordEvent(TEXT("Core.DDC.Summary"), EventAttributes);
 #endif			
 }
 
-void FEditorTelemetry::RecordEvent_DDCVerification(const FString& Context, TArray<FAnalyticsEventAttribute> Attributes)
+void FEditorTelemetry::RecordEvent_DDCVerification(const FString& Context, const TArray<FAnalyticsEventAttribute>& Attributes)
 {
 	const int SchemaVersion = 1;
-	Attributes.Emplace(TEXT("SchemaVersion"), SchemaVersion);
-	Attributes.Emplace(TEXT("Context"), Context);
 
 	FDerivedDataCacheVerificationStats DDCVerifyStats;
 	GetDerivedDataCacheRef().GatherVerificationStats(DDCVerifyStats);
@@ -260,6 +292,9 @@ void FEditorTelemetry::RecordEvent_DDCVerification(const FString& Context, TArra
 		if (Stat.RecordsTested > 0 || Stat.ValuesTested > 0)
 		{
 			TArray<FAnalyticsEventAttribute> EventAttributes = Attributes;
+
+			EventAttributes.Emplace(TEXT("SchemaVersion"), SchemaVersion);
+			EventAttributes.Emplace(TEXT("Context"), Context);
 			EventAttributes.Emplace(TEXT("Bucket"), Stat.BucketName);
 			EventAttributes.Emplace(TEXT("RecordsTested"), Stat.RecordsTested);
 			EventAttributes.Emplace(TEXT("RecordsMatchedCache"), Stat.RecordsMatchedCache);
@@ -279,7 +314,7 @@ void FEditorTelemetry::RecordEvent_DDCVerification(const FString& Context, TArra
 	}
 }
 
-void FEditorTelemetry::RecordEvent_Zen(const FString& Context, TArray<FAnalyticsEventAttribute> Attributes)
+void FEditorTelemetry::RecordEvent_Zen(const FString& Context, const TArray<FAnalyticsEventAttribute>& Attributes)
 {
 #if UE_WITH_ZEN
 	// Gather Zen analytics
@@ -287,44 +322,50 @@ void FEditorTelemetry::RecordEvent_Zen(const FString& Context, TArray<FAnalytics
 	{
 		const int SchemaVersion = 2;
 
-		Attributes.Emplace(TEXT("SchemaVersion"), SchemaVersion);
-		Attributes.Emplace(TEXT("Context"), Context);
+		TArray<FAnalyticsEventAttribute> EventAttributes = Attributes;
+
+		EventAttributes.Emplace(TEXT("SchemaVersion"), SchemaVersion);
+		EventAttributes.Emplace(TEXT("Context"), Context);
 
-		UE::Zen::GetDefaultServiceInstance().GatherAnalytics(Attributes);
-		FStudioTelemetry::Get().RecordEvent(TEXT("Core.Zen"), Attributes);
+		UE::Zen::GetDefaultServiceInstance().GatherAnalytics(EventAttributes);
+		FStudioTelemetry::Get().RecordEvent(TEXT("Core.Zen"), EventAttributes);
 	}
 #endif
 }
 
-void FEditorTelemetry::RecordEvent_VirtualAssets(const FString& Context, TArray<FAnalyticsEventAttribute> Attributes)
+void FEditorTelemetry::RecordEvent_VirtualAssets(const FString& Context, const TArray<FAnalyticsEventAttribute>& Attributes)
 {
 	if (UE::Virtualization::IVirtualizationSystem::Get().IsEnabled())
 	{
 		const int SchemaVersion = 2;
 
-		Attributes.Emplace(TEXT("SchemaVersion"), SchemaVersion);
-		Attributes.Emplace(TEXT("Context"), Context);
+		TArray<FAnalyticsEventAttribute> EventAttributes = Attributes;
+
+		EventAttributes.Emplace(TEXT("SchemaVersion"), SchemaVersion);
+		EventAttributes.Emplace(TEXT("Context"), Context);
 
 		// Gather Virtualization analytics
-		UE::Virtualization::IVirtualizationSystem::Get().GatherAnalytics(Attributes);
+		UE::Virtualization::IVirtualizationSystem::Get().GatherAnalytics(EventAttributes);
 
-		FStudioTelemetry::Get().RecordEvent(TEXT("Core.VirtualAssets"), Attributes);
+		FStudioTelemetry::Get().RecordEvent(TEXT("Core.VirtualAssets"), EventAttributes);
 	}
 }
 
-void FEditorTelemetry::RecordEvent_ShaderStats(const FString& Context, TArray<FAnalyticsEventAttribute> Attributes)
+void FEditorTelemetry::RecordEvent_ShaderStats(const FString& Context, const TArray<FAnalyticsEventAttribute>& Attributes)
 {
 	const int SchemaVersion = 1;
 
-	Attributes.Emplace(TEXT("SchemaVersion"), SchemaVersion);
-	Attributes.Emplace(TEXT("Context"), Context);
+	TArray<FAnalyticsEventAttribute> EventAttributes = Attributes;
+
+	EventAttributes.Emplace(TEXT("SchemaVersion"), SchemaVersion);
+	EventAttributes.Emplace(TEXT("Context"), Context);
 
-	FShaderStatsFunctions::GatherShaderAnalytics(Attributes);
+	FShaderStatsFunctions::GatherShaderAnalytics(EventAttributes);
 
-	FStudioTelemetry::Get().RecordEvent(TEXT("Core.ShaderStats"), Attributes);
+	FStudioTelemetry::Get().RecordEvent(TEXT("Core.ShaderStats"), EventAttributes);
 }
 
-void FEditorTelemetry::RecordEvent_MemoryLLM(const FString& Context, TArray<FAnalyticsEventAttribute> Attributes)
+void FEditorTelemetry::RecordEvent_MemoryLLM(const FString& Context, const TArray<FAnalyticsEventAttribute>& Attributes)
 {
 #if ENABLE_LOW_LEVEL_MEM_TRACKER
 	if (FStudioTelemetry::Get().IsSessionRunning())
@@ -365,7 +406,7 @@ void FEditorTelemetry::RecordEvent_MemoryLLM(const FString& Context, TArray<FAna
 #endif
 }
 
-void FEditorTelemetry::RecordEvent_CoreSystems(const FString& Context, TArray<FAnalyticsEventAttribute> Attributes)
+void FEditorTelemetry::RecordEvent_CoreSystems(const FString& Context, const TArray<FAnalyticsEventAttribute>& Attributes)
 {
 	FEditorTelemetry::RecordEvent_DDCResource(Context, Attributes);
 	FEditorTelemetry::RecordEvent_DDCSummary(Context, Attributes);
@@ -587,7 +628,7 @@ void FEditorTelemetry::StartSession()
 				PIESpan = FStudioTelemetry::Get().StartSpan(PIESpanName);
 
 				// Append the PIE transition count to the PIE name
-				PIEStartupSpan = FStudioTelemetry::Get().StartSpan(PIESessionCount == 0 ? PIEStartupSpanName : FName(*FString::Printf(TEXT("%s%d"), *PIEStartupSpanName.ToString(), PIESessionCount)), PIESpan);
+				PIEStartupSpan = FStudioTelemetry::Get().StartSpan(PIESessionCount == 0 ? PIEStartupSpanName : FName(*FString::Printf(TEXT("%s.Iterate"), *PIEStartupSpanName.ToString())), PIESpan);
 
 				TArray<FAnalyticsEventAttribute> Attributes;
 				Attributes.Emplace(TEXT("MapName"), EditorMapName);
@@ -598,7 +639,7 @@ void FEditorTelemetry::StartSession()
 
 		FEditorDelegates::PreBeginPIE.AddLambda([this](bool)
 			{
-				PIEPreBeginSpan = FStudioTelemetry::Get().StartSpan(PIESessionCount == 0 ? PIEPreBeginSpanName : FName(*FString::Printf(TEXT("%s%d"), *PIEPreBeginSpanName.ToString(), PIESessionCount)), PIESpan);
+				PIEPreBeginSpan = FStudioTelemetry::Get().StartSpan(PIESessionCount == 0 ? PIEPreBeginSpanName : FName(*FString::Printf(TEXT("%s.Iterate"), *PIEPreBeginSpanName.ToString())), PIESpan);
 				PIEPreBeginSpan->AddAttributes(PIESpan->GetAttributes());
 			});
 
@@ -663,7 +704,7 @@ void FEditorTelemetry::StartSession()
 						}
 					}
 
-					PIEInteractSpan = FStudioTelemetry::Get().StartSpan(PIESessionCount == 0 ? PIEInteractSpanName : FName(*FString::Printf(TEXT("%s%d"), *PIEInteractSpanName.ToString(), PIESessionCount)), PIESpan);
+					PIEInteractSpan = FStudioTelemetry::Get().StartSpan(PIESessionCount == 0 ? PIEInteractSpanName : FName(*FString::Printf(TEXT("%s.Iterate"), *PIEInteractSpanName.ToString())), PIESpan);
 					PIEInteractSpan->AddAttributes(PIESpan->GetAttributes());
 				}
 			});
@@ -674,7 +715,7 @@ void FEditorTelemetry::StartSession()
 				{
 					// PIE is ending so no longer interactive
 					FStudioTelemetry::Get().EndSpan(PIEInteractSpan);
-					PIEShutdownSpan = FStudioTelemetry::Get().StartSpan(PIESessionCount == 0 ? PIEShutdownSpanName : FName(*FString::Printf(TEXT("%s%d"), *PIEShutdownSpanName.ToString(), PIESessionCount)), PIESpan);
+					PIEShutdownSpan = FStudioTelemetry::Get().StartSpan(PIESessionCount == 0 ? PIEShutdownSpanName : FName(*FString::Printf(TEXT("%s.Iterate"), *PIEShutdownSpanName.ToString())), PIESpan);
 					PIEShutdownSpan->AddAttributes(PIESpan->GetAttributes());
 				}
 			});
diff --git a/Engine/Plugins/Experimental/EditorTelemetry/Source/Public/EditorTelemetry.h b/Engine/Plugins/Experimental/EditorTelemetry/Source/Public/EditorTelemetry.h
--- a/Engine/Plugins/Experimental/EditorTelemetry/Source/Public/EditorTelemetry.h
+++ b/Engine/Plugins/Experimental/EditorTelemetry/Source/Public/EditorTelemetry.h
@@ -32,16 +32,16 @@ public:
 	UE_API void EndSession();
 
 	/** Useful event recording functions */
-	UE_API void RecordEvent_Cooking(const FString& Context, TArray<FAnalyticsEventAttribute> Attributes = {});
-	UE_API void RecordEvent_Loading(const FString& Context, double LoadingSeconds, TArray<FAnalyticsEventAttribute> Attributes = {});
-	UE_API void RecordEvent_CoreSystems(const FString& Context, TArray<FAnalyticsEventAttribute> Attributes = {});
-	UE_API void RecordEvent_DDCResource(const FString& Context, TArray<FAnalyticsEventAttribute> Attributes = {});
-	UE_API void RecordEvent_DDCSummary(const FString& Context, TArray<FAnalyticsEventAttribute> Attributes = {});
-	UE_API void RecordEvent_DDCVerification(const FString& Context, TArray<FAnalyticsEventAttribute> Attributes = {});
-	UE_API void RecordEvent_Zen(const FString& Context, TArray<FAnalyticsEventAttribute> Attributes = {});
-	UE_API void RecordEvent_VirtualAssets(const FString& Context, TArray<FAnalyticsEventAttribute> Attributes = {});
-	UE_API void RecordEvent_MemoryLLM(const FString& Context, TArray<FAnalyticsEventAttribute> Attributes = {});
-	UE_API void RecordEvent_ShaderStats(const FString& Context, TArray<FAnalyticsEventAttribute> Attributes = {});
+	UE_API void RecordEvent_Cooking(const FString& Context, const TArray<FAnalyticsEventAttribute>& Attributes = {});
+	UE_API void RecordEvent_Loading(const FString& Context, double LoadingSeconds, const TArray<FAnalyticsEventAttribute>& Attributes = {});
+	UE_API void RecordEvent_CoreSystems(const FString& Context, const TArray<FAnalyticsEventAttribute>& Attributes = {});
+	UE_API void RecordEvent_DDCResource(const FString& Context, const TArray<FAnalyticsEventAttribute>& Attributes = {});
+	UE_API void RecordEvent_DDCSummary(const FString& Context, const TArray<FAnalyticsEventAttribute>& Attributes = {});
+	UE_API void RecordEvent_DDCVerification(const FString& Context, const TArray<FAnalyticsEventAttribute>& Attributes = {});
+	UE_API void RecordEvent_Zen(const FString& Context, const TArray<FAnalyticsEventAttribute>& Attributes = {});
+	UE_API void RecordEvent_VirtualAssets(const FString& Context, const TArray<FAnalyticsEventAttribute>& Attributes = {});
+	UE_API void RecordEvent_MemoryLLM(const FString& Context, const TArray<FAnalyticsEventAttribute>& Attributes = {});
+	UE_API void RecordEvent_ShaderStats(const FString& Context, const TArray<FAnalyticsEventAttribute>& Attributes = {});
 
 	UE_API void RegisterCollectionWorkflowDelegates(FTelemetryRouter& Router);
 
diff --git a/Engine/Source/Developer/DerivedDataCache/Private/DerivedDataCache.cpp b/Engine/Source/Developer/DerivedDataCache/Private/DerivedDataCache.cpp
--- a/Engine/Source/Developer/DerivedDataCache/Private/DerivedDataCache.cpp
+++ b/Engine/Source/Developer/DerivedDataCache/Private/DerivedDataCache.cpp
@@ -899,96 +899,6 @@ public:
 		GatherDerivedDataCacheSummaryStats(DDCSummaryStats);
 	}
 
-	virtual void GatherAnalytics(TArray<FAnalyticsEventAttribute>& Attributes) const override
-	{
-#if ENABLE_COOK_STATS
-
-		// Gather the latest resource stats
-		TArray<FDerivedDataCacheResourceStat> ResourceStats;
-
-		GatherResourceStats(ResourceStats);
-
-		// Append to the attributes
-		for (const FDerivedDataCacheResourceStat& Stat : ResourceStats)
-		{
-			FString BaseName = TEXT("DDC_Resource_") + Stat.AssetType;
-
-			BaseName = BaseName.Replace(TEXT("("), TEXT("")).Replace(TEXT(")"), TEXT(""));
-
-			{
-				FString AttrName = BaseName + TEXT("_BuildCount");
-				Attributes.Emplace(MoveTemp(AttrName), Stat.BuildCount);
-			}
-
-			{
-				FString AttrName = BaseName + TEXT("_BuildTimeSec");
-				Attributes.Emplace(MoveTemp(AttrName), Stat.BuildTimeSec);
-			}
-
-			{
-				FString AttrName = BaseName + TEXT("_BuildSizeMB");
-				Attributes.Emplace(MoveTemp(AttrName), Stat.BuildSizeMB);
-			}
-
-			{
-				FString AttrName = BaseName + TEXT("_LoadCount");
-				Attributes.Emplace(MoveTemp(AttrName), Stat.LoadCount);
-			}
-
-			{
-				FString AttrName = BaseName + TEXT("_LoadTimeSecLoadTimeSec");
-				Attributes.Emplace(MoveTemp(AttrName), Stat.LoadTimeSec);
-			}
-
-			{
-				FString AttrName = BaseName + TEXT("_LoadSizeMB");
-				Attributes.Emplace(MoveTemp(AttrName), Stat.LoadSizeMB);
-			}
-		}
-
-		// Gather the summary stats
-		FDerivedDataCacheSummaryStats SummaryStats;
-
-		GatherDerivedDataCacheSummaryStats(SummaryStats);
-
-		// Append to the attributes
-		for (const FDerivedDataCacheSummaryStat& Stat : SummaryStats.Stats)
-		{
-			FString FormattedAttrName = "DDC_Summary_" + Stat.Key.Replace(TEXT("."), TEXT("_"));
-
-			if (Stat.Value.IsNumeric())
-			{
-				Attributes.Emplace(FormattedAttrName, FCString::Atof(*Stat.Value));
-			}
-			else
-			{
-				Attributes.Emplace(FormattedAttrName, Stat.Value);
-			}
-		}
-
-		// Gather the backend stats
-		TSharedRef<FDerivedDataCacheStatsNode> RootNode = Backend->GatherUsageStats();
-		RootNode->ForEachDescendant([&Attributes](TSharedRef<const FDerivedDataCacheStatsNode> Node)
-			{
-				const FString& CacheName = Node->GetCacheName();
-
-				for (const FCookStatsManager::StringKeyValue& Stat : Node->CustomStats)
-				{
-					FString FormattedAttrName = CacheName + TEXT("_") + Stat.Key.Replace(TEXT("."), TEXT("_"));
-
-					if (Stat.Value.IsNumeric())
-					{
-						Attributes.Emplace(FormattedAttrName, FCString::Atof(*Stat.Value));
-					}
-					else
-					{
-						Attributes.Emplace(FormattedAttrName, Stat.Value);
-					}
-				}
-			});
-#endif
-	}
-
 	void GatherVerificationStats(FDerivedDataCacheVerificationStats& OutStats) const override
 	{
 		Backend->GatherVerificationStats(OutStats);
diff --git a/Engine/Source/Developer/DerivedDataCache/Public/DerivedDataCacheInterface.h b/Engine/Source/Developer/DerivedDataCache/Public/DerivedDataCacheInterface.h
--- a/Engine/Source/Developer/DerivedDataCache/Public/DerivedDataCacheInterface.h
+++ b/Engine/Source/Developer/DerivedDataCache/Public/DerivedDataCacheInterface.h
@@ -284,9 +284,6 @@ public:
 	 */
 	UE_INTERNAL virtual void GatherVerificationStats(FDerivedDataCacheVerificationStats& OutStats) const = 0;
 
-	UE_DEPRECATED(5.4, "This is no longer used. Use GatherUsageStats/GatherResourceStats/GatherSummaryStats instead")
-	virtual void GatherAnalytics(TArray<FAnalyticsEventAttribute>& Attributes) const = 0;
-
 	//-----------------------
 	// Notification Interface
 	//-----------------------
diff --git a/Engine/Source/Runtime/StudioTelemetry/Private/StudioTelemetry.cpp b/Engine/Source/Runtime/StudioTelemetry/Private/StudioTelemetry.cpp
--- a/Engine/Source/Runtime/StudioTelemetry/Private/StudioTelemetry.cpp
+++ b/Engine/Source/Runtime/StudioTelemetry/Private/StudioTelemetry.cpp
@@ -22,7 +22,7 @@ DEFINE_LOG_CATEGORY(LogStudioTelemetry);
 
 IMPLEMENT_MODULE(FStudioTelemetry, StudioTelemetry)
 
-FString GetEnvironmentValriable(const FString& CommandLineVar, const FString& EnvironmentVar)
+FString GetEnvironmentVariable(const FString& CommandLineVar, const FString& EnvironmentVar)
 {
 	FString Result;
 	if (false == FParse::Value(FCommandLine::Get(), *CommandLineVar, Result))
@@ -156,16 +156,17 @@ bool FStudioTelemetry::StartSession()
 			SessionContextAttributes.Emplace(TEXT("OS_ID"), FPlatformMisc::GetOperatingSystemId());
 		}
 
-		const FString HordeJobId = GetEnvironmentValriable(TEXT("HordeJobId="), TEXT("UE_HORDE_JOBID")) ;
+		const FString HordeJobId = GetEnvironmentVariable(TEXT("HordeJobId="), TEXT("UE_HORDE_JOBID"));
+		const FString HordeServerURL = GetEnvironmentVariable(TEXT("HordeServerUrl="), TEXT("UE_HORDE_URL") );
 
 		// Send Horde environment if available
-		if (!HordeJobId.IsEmpty())
+		if (!HordeServerURL.IsEmpty() && !HordeJobId.IsEmpty())
 		{
+			SessionContextAttributes.Emplace(TEXT("Horde_ServerURL"), HordeServerURL);
 			SessionContextAttributes.Emplace(TEXT("Horde_JobID"), HordeJobId);
-			SessionContextAttributes.Emplace(TEXT("Horde_StepID"), GetEnvironmentValriable(TEXT("HordeStepId="), TEXT("UE_HORDE_STEPID")));
-			SessionContextAttributes.Emplace(TEXT("Horde_StepName"), GetEnvironmentValriable(TEXT("HordeStepName="), TEXT("UE_HORDE_STEPNAME")));
-			SessionContextAttributes.Emplace(TEXT("Horde_ServerURL"), GetEnvironmentValriable(TEXT("HordeServerUrl="), TEXT("UE_HORDE_URL")));
-			SessionContextAttributes.Emplace(TEXT("Horde_TemplateName"), GetEnvironmentValriable(TEXT("HordeTemplateName="), TEXT("UE_HORDE_TEMPLATENAME")));	
+			SessionContextAttributes.Emplace(TEXT("Horde_StepID"), GetEnvironmentVariable(TEXT("HordeStepId="), TEXT("UE_HORDE_STEPID")));
+			SessionContextAttributes.Emplace(TEXT("Horde_StepName"), GetEnvironmentVariable(TEXT("HordeStepName="), TEXT("UE_HORDE_STEPNAME")));
+			SessionContextAttributes.Emplace(TEXT("Horde_TemplateName"), GetEnvironmentVariable(TEXT("HordeTemplateName="), TEXT("UE_HORDE_TEMPLATENAME")));	
 		}
 
 		// Send the Session Context attributes with ALL events by appending to the default attributes if required
