Fixed missing Shader Stats from Cook Event

#Synced-CL 50219800

diff --git a/Engine/Plugins/Experimental/EditorTelemetry/Source/Private/EditorTelemetry.cpp b/Engine/Plugins/Experimental/EditorTelemetry/Source/Private/EditorTelemetry.cpp
--- a/Engine/Plugins/Experimental/EditorTelemetry/Source/Private/EditorTelemetry.cpp
+++ b/Engine/Plugins/Experimental/EditorTelemetry/Source/Private/EditorTelemetry.cpp
@@ -155,6 +155,9 @@ void FEditorTelemetry::RecordEvent_Cooking(const FString& Context, const TArray<
 			FullEventAttributes.Emplace((*it).Value);
 		}
 
+		// Append the cook stats
+		FShaderStatsFunctions::GatherShaderAnalytics(FullEventAttributes);
+
 		// Record the full cook event
 		FStudioTelemetry::Get().RecordEvent(TEXT("Core.Cooking"), FullEventAttributes);
 	}
